Editgrate 4.25 audio features from project stream:

diff --git a/Engine/Source/Editor/AudioEditor/AudioEditor.Build.cs b/Engine/Source/Editor/AudioEditor/AudioEditor.Build.cs
--- a/Engine/Source/Editor/AudioEditor/AudioEditor.Build.cs
+++ b/Engine/Source/Editor/AudioEditor/AudioEditor.Build.cs
@@ -17,6 +17,7 @@ public class AudioEditor : ModuleRules
 			new string[] {
 				"AudioMixer",
 				"ToolMenus",
+				"AudioExtensions"
 			});
 
 		PublicDependencyModuleNames.AddRange(
diff --git a/Engine/Source/Editor/AudioEditor/Classes/Factories/SoundSubmixFactory.h b/Engine/Source/Editor/AudioEditor/Classes/Factories/SoundSubmixFactory.h
--- a/Engine/Source/Editor/AudioEditor/Classes/Factories/SoundSubmixFactory.h
+++ b/Engine/Source/Editor/AudioEditor/Classes/Factories/SoundSubmixFactory.h
@@ -8,8 +8,8 @@
 #include "Factories/Factory.h"
 #include "SoundSubmixFactory.generated.h"
 
-UCLASS(MinimalAPI, hidecategories=Object)
-class USoundSubmixFactory : public UFactory
+UCLASS(hidecategories=Object)
+class AUDIOEDITOR_API USoundSubmixFactory : public UFactory
 {
 	GENERATED_UCLASS_BODY()
 
@@ -18,3 +18,36 @@ class USoundSubmixFactory : public UFactory
 	virtual bool CanCreateNew() const override;
 	//~ Begin UFactory Interface	
 };
+
+UCLASS(hidecategories = Object)
+class AUDIOEDITOR_API USoundfieldSubmixFactory: public UFactory
+{
+	GENERATED_UCLASS_BODY()
+
+	//~ Begin UFactory Interface
+	virtual UObject* FactoryCreateNew(UClass* Class, UObject* InParent, FName Name, EObjectFlags Flags, UObject* Context, FFeedbackContext* Warn) override;
+	virtual bool CanCreateNew() const override;
+	//~ Begin UFactory Interface	
+};
+
+UCLASS(hidecategories = Object)
+class AUDIOEDITOR_API UEndpointSubmixFactory : public UFactory
+{
+	GENERATED_UCLASS_BODY()
+
+	//~ Begin UFactory Interface
+	virtual UObject* FactoryCreateNew(UClass* Class, UObject* InParent, FName Name, EObjectFlags Flags, UObject* Context, FFeedbackContext* Warn) override;
+	virtual bool CanCreateNew() const override;
+	//~ Begin UFactory Interface	
+};
+
+UCLASS(hidecategories = Object)
+class AUDIOEDITOR_API USoundfieldEndpointSubmixFactory : public UFactory
+{
+	GENERATED_UCLASS_BODY()
+
+	//~ Begin UFactory Interface
+	virtual UObject* FactoryCreateNew(UClass* Class, UObject* InParent, FName Name, EObjectFlags Flags, UObject* Context, FFeedbackContext* Warn) override;
+	virtual bool CanCreateNew() const override;
+	//~ Begin UFactory Interface	
+};
diff --git a/Engine/Source/Editor/AudioEditor/Classes/SoundSubmixGraph/SoundSubmixGraph.h b/Engine/Source/Editor/AudioEditor/Classes/SoundSubmixGraph/SoundSubmixGraph.h
--- a/Engine/Source/Editor/AudioEditor/Classes/SoundSubmixGraph/SoundSubmixGraph.h
+++ b/Engine/Source/Editor/AudioEditor/Classes/SoundSubmixGraph/SoundSubmixGraph.h
@@ -6,7 +6,7 @@
 
 
 // Forward Declarations
-class USoundSubmix;
+class USoundSubmixBase;
 class USoundSubmixGraphNode;
 
 
@@ -19,12 +19,12 @@ public:
 	/**
 	 * Set the SoundSubmix which forms the root of this graph
 	 */
-	AUDIOEDITOR_API void SetRootSoundSubmix(USoundSubmix* InSoundSubmix);
+	AUDIOEDITOR_API void SetRootSoundSubmix(USoundSubmixBase* InSoundSubmix);
 
 	/**
 	 * Get the SoundSubmix which forms the root of this graph
 	 */
-	USoundSubmix* GetRootSoundSubmix() const;
+	USoundSubmixBase* GetRootSoundSubmix() const;
 
 	/**
 	 * Completely rebuild the graph from the root, removing all old nodes
@@ -38,7 +38,7 @@ public:
 	 * @param	NodePosX		X coordinate submix(es) were dropped at
 	 * @param	NodePosY		Y coordinate submix(es) were dropped at
 	 */
-	void AddDroppedSoundSubmixes(const TSet<USoundSubmix*>& SoundSubmixes, int32 NodePosX, int32 NodePosY);
+	void AddDroppedSoundSubmixes(const TSet<USoundSubmixBase*>& SoundSubmixes, int32 NodePosX, int32 NodePosY);
 
 	/**
 	 * Display a new SoundSubmix that has just been created using the editor
@@ -49,12 +49,12 @@ public:
 	 * @param	NodePosY	Y coordinate new submix was created at
 	 * @param bSelectNewNode	Whether or not to select the new node being created
 	 */
-	AUDIOEDITOR_API void AddNewSoundSubmix(UEdGraphPin* FromPin, USoundSubmix* SoundSubmix, int32 NodePosX, int32 NodePosY, bool bSelectNewNode = true);
+	AUDIOEDITOR_API void AddNewSoundSubmix(UEdGraphPin* FromPin, USoundSubmixBase* SoundSubmix, int32 NodePosX, int32 NodePosY, bool bSelectNewNode = true);
 
 	/**
 	 * Checks whether a SoundSubmix is already represented on this graph
 	 */
-	bool IsSubmixDisplayed(USoundSubmix* SoundSubmix) const;
+	bool IsSubmixDisplayed(USoundSubmixBase* SoundSubmix) const;
 
 	/**
 	 * Use this graph to re-link all of the SoundSubmixes it represents after a change in linkage
@@ -75,7 +75,7 @@ public:
 	/**
 	 * Find an existing node that represents a given SoundSubmix
 	 */
-	AUDIOEDITOR_API USoundSubmixGraphNode* FindExistingNode(USoundSubmix* SoundSubmix) const;
+	AUDIOEDITOR_API USoundSubmixGraphNode* FindExistingNode(USoundSubmixBase* SoundSubmix) const;
 
 private:
 	/**
@@ -87,7 +87,7 @@ private:
 	 * @param bSelectNewNode	Whether or not to select the new node being created
 	 * @return	Total height of all constructed nodes (used to arrange multiple new nodes)
 	 */
-	int32 ConstructNodes(USoundSubmix* SoundSubmix, int32 NodePosX, int32 NodePosY, bool bSelectNewNode = true);
+	int32 ConstructNodes(USoundSubmixBase* SoundSubmix, int32 NodePosX, int32 NodePosY, bool bSelectNewNode = true);
 	/**
 	 * Recursively build a map of child counts for each SoundSubmix to arrange them correctly
 	 *
@@ -95,7 +95,7 @@ private:
 	 * @param	OutChildCounts	Map of child counts
 	 * @return	Total child count for ParentClass
 	 */
-	int32 RecursivelyGatherChildCounts(USoundSubmix* ParentSubmix, TMap<USoundSubmix*, int32>& OutChildCounts);
+	int32 RecursivelyGatherChildCounts(USoundSubmixBase* ParentSubmix, TMap<USoundSubmixBase*, int32>& OutChildCounts);
 	/**
 	 * Recursively Construct Nodes to represent the children of a SoundSubmix
 	 *
@@ -104,7 +104,7 @@ private:
 	 * @param bSelectNewNode	Whether or not to select the new node being created
 	 * @return	Total height of constructed nodes (used to arrange next new node)
 	 */
-	int32 RecursivelyConstructChildNodes(USoundSubmixGraphNode* ParentNode, const TMap<USoundSubmix*, int32>& InChildCounts, bool bSelectNewNode = true);
+	int32 RecursivelyConstructChildNodes(USoundSubmixGraphNode* ParentNode, const TMap<USoundSubmixBase*, int32>& InChildCounts, bool bSelectNewNode = true);
 	/**
 	 * Recursively remove a node and its children from the graph
 	 */
@@ -122,10 +122,10 @@ private:
 	 * @param bSelectNewNode	Whether or not to select the new node being created
 	 * @return	Either a new node or an existing node representing the class
 	 */
-	USoundSubmixGraphNode* CreateNode(USoundSubmix* SoundSubmix, int32 NodePosX, int32 NodePosY, bool bSelectNewNode = true);
+	USoundSubmixGraphNode* CreateNode(USoundSubmixBase* SoundSubmix, int32 NodePosX, int32 NodePosY, bool bSelectNewNode = true);
 
 private:
 	/** SoundSubmix which forms the root of this graph */
-	USoundSubmix*	RootSoundSubmix;
+	USoundSubmixBase*	RootSoundSubmix;
 };
 
diff --git a/Engine/Source/Editor/AudioEditor/Classes/SoundSubmixGraph/SoundSubmixGraphNode.h b/Engine/Source/Editor/AudioEditor/Classes/SoundSubmixGraph/SoundSubmixGraphNode.h
--- a/Engine/Source/Editor/AudioEditor/Classes/SoundSubmixGraph/SoundSubmixGraphNode.h
+++ b/Engine/Source/Editor/AudioEditor/Classes/SoundSubmixGraph/SoundSubmixGraphNode.h
@@ -7,7 +7,7 @@
 #include "EdGraph/EdGraphNode.h"
 #include "SoundSubmixGraphNode.generated.h"
 
-class USoundSubmix;
+class USoundSubmixBase;
 
 UCLASS(MinimalAPI)
 class USoundSubmixGraphNode : public UEdGraphNode
@@ -16,7 +16,7 @@ class USoundSubmixGraphNode : public UEdGraphNode
 
 	/** The SoundSubmix this represents */
 	UPROPERTY(VisibleAnywhere, instanced, Category=Sound)
-	USoundSubmix* SoundSubmix;
+	USoundSubmixBase* SoundSubmix;
 
 	/** Get the Pin that connects to all children */
 	UEdGraphPin* GetChildPin() const { return ChildPin; }
diff --git a/Engine/Source/Editor/AudioEditor/Private/AssetTypeActions/AssetTypeActions_SoundSubmix.cpp b/Engine/Source/Editor/AudioEditor/Private/AssetTypeActions/AssetTypeActions_SoundSubmix.cpp
--- a/Engine/Source/Editor/AudioEditor/Private/AssetTypeActions/AssetTypeActions_SoundSubmix.cpp
+++ b/Engine/Source/Editor/AudioEditor/Private/AssetTypeActions/AssetTypeActions_SoundSubmix.cpp
@@ -19,7 +19,7 @@ void FAssetTypeActions_SoundSubmix::OpenAssetEditor(const TArray<UObject*>& InOb
 
 	for (UObject* Obj : InObjects)
 	{
-		if (USoundSubmix* SoundSubmix = Cast<USoundSubmix>(Obj))
+		if (USoundSubmixBase* SoundSubmix = Cast<USoundSubmixBase>(Obj))
 		{
 			IAudioEditorModule* AudioEditorModule = &FModuleManager::LoadModuleChecked<IAudioEditorModule>("AudioEditor");
 			AudioEditorModule->CreateSoundSubmixEditor(Mode, EditWithinLevelEditor, SoundSubmix);
@@ -29,11 +29,11 @@ void FAssetTypeActions_SoundSubmix::OpenAssetEditor(const TArray<UObject*>& InOb
 
 bool FAssetTypeActions_SoundSubmix::AssetsActivatedOverride(const TArray<UObject*>& InObjects, EAssetTypeActivationMethod::Type ActivationType)
 {
-	TSet<USoundSubmix*> SubmixesToSelect;
+	TSet<USoundSubmixBase*> SubmixesToSelect;
 	IAssetEditorInstance* Editor = nullptr;
 	for (UObject* Obj : InObjects)
 	{
-		if (USoundSubmix* SubmixToSelect = Cast<USoundSubmix>(Obj))
+		if (USoundSubmixBase* SubmixToSelect = Cast<USoundSubmixBase>(Obj))
 		{
 			if (!Editor)
 			{
@@ -61,4 +61,45 @@ const TArray<FText>& FAssetTypeActions_SoundSubmix::GetSubMenus() const
 
 	return SubMenus;
 }
-#undef LOCTEXT_NAMESPACE
\ No newline at end of file
+
+UClass* FAssetTypeActions_SoundfieldSubmix::GetSupportedClass() const
+{
+	return USoundfieldSubmix::StaticClass();
+}
+
+UClass* FAssetTypeActions_EndpointSubmix::GetSupportedClass() const
+{
+	return UEndpointSubmix::StaticClass();
+}
+
+UClass* FAssetTypeActions_SoundfieldEndpointSubmix::GetSupportedClass() const
+{
+	return USoundfieldEndpointSubmix::StaticClass();
+}
+
+UClass* FAssetTypeActions_SoundfieldEncodingSettings::GetSupportedClass() const
+{
+	return USoundfieldEncodingSettingsBase::StaticClass();
+}
+
+UClass* FAssetTypeActions_SoundfieldEffectSettings::GetSupportedClass() const
+{
+	return USoundfieldEncodingSettingsBase::StaticClass();
+}
+
+UClass* FAssetTypeActions_SoundfieldEffect::GetSupportedClass() const
+{
+	return USoundfieldEffectBase::StaticClass();
+}
+
+UClass* FAssetTypeActions_AudioEndpointSettings::GetSupportedClass() const
+{
+	return UAudioEndpointSettingsBase::StaticClass();
+}
+
+UClass* FAssetTypeActions_SoundfieldEndpointSettings::GetSupportedClass() const
+{
+	return USoundfieldEndpointSettingsBase::StaticClass();
+}
+
+#undef LOCTEXT_NAMESPACE
diff --git a/Engine/Source/Editor/AudioEditor/Private/AssetTypeActions/AssetTypeActions_SoundSubmix.h b/Engine/Source/Editor/AudioEditor/Private/AssetTypeActions/AssetTypeActions_SoundSubmix.h
--- a/Engine/Source/Editor/AudioEditor/Private/AssetTypeActions/AssetTypeActions_SoundSubmix.h
+++ b/Engine/Source/Editor/AudioEditor/Private/AssetTypeActions/AssetTypeActions_SoundSubmix.h
@@ -3,17 +3,107 @@
 #pragma once
 
 #include "AssetTypeActions_Base.h"
+#include "SoundSubmixDefaultColorPalette.h"
+
+/** Submix Types: */
 
 class FAssetTypeActions_SoundSubmix : public FAssetTypeActions_Base
 {
 public:
 	// IAssetTypeActions Implementation
 	virtual FText GetName() const override { return NSLOCTEXT("AssetTypeActions", "AssetTypeActions_SoundSubmix", "Sound Submix"); }
-	virtual FColor GetTypeColor() const override { return FColor(175, 255, 0); }
+	virtual FColor GetTypeColor() const override { return Audio::DefaultSubmixColor; }
 	virtual const TArray<FText>& GetSubMenus() const override;
 	virtual UClass* GetSupportedClass() const override;
 
 	virtual bool AssetsActivatedOverride(const TArray<UObject*>& InObjects, EAssetTypeActivationMethod::Type ActivationType) override;
 	virtual void OpenAssetEditor( const TArray<UObject*>& InObjects, TSharedPtr<class IToolkitHost> EditWithinLevelEditor = TSharedPtr<IToolkitHost>() ) override;
 	virtual uint32 GetCategories() override { return EAssetTypeCategories::Sounds; }
-};
\ No newline at end of file
+};
+
+class FAssetTypeActions_SoundfieldSubmix : public FAssetTypeActions_SoundSubmix
+{
+public:
+	// IAssetTypeActions Implementation
+	virtual FText GetName() const override { return NSLOCTEXT("AssetTypeActions", "AssetTypeActions_SoundfieldSubmix", "Soundfield Submix"); }
+	virtual FColor GetTypeColor() const override { return Audio::SoundfieldDefaultSubmixColor; }
+	virtual UClass* GetSupportedClass() const override;
+	virtual uint32 GetCategories() override { return EAssetTypeCategories::Sounds; }
+};
+
+class FAssetTypeActions_EndpointSubmix : public FAssetTypeActions_SoundSubmix
+{
+public:
+	// IAssetTypeActions Implementation
+	virtual FText GetName() const override { return NSLOCTEXT("AssetTypeActions", "AssetTypeActions_EndpointSubmix", "Endpoint Submix"); }
+	virtual FColor GetTypeColor() const override { return Audio::EndpointDefaultSubmixColor; }
+	virtual UClass* GetSupportedClass() const override;
+	virtual uint32 GetCategories() override { return EAssetTypeCategories::Sounds; }
+};
+
+class FAssetTypeActions_SoundfieldEndpointSubmix : public FAssetTypeActions_SoundSubmix
+{
+public:
+	// IAssetTypeActions Implementation
+	virtual FText GetName() const override { return NSLOCTEXT("AssetTypeActions", "AssetTypeActions_SoundfieldEndpointSubmix", "Soundfield Endpoint Submix"); }
+	virtual FColor GetTypeColor() const override { return Audio::SoundfieldEndpointDefaultSubmixColor; }
+	virtual UClass* GetSupportedClass() const override;
+	virtual uint32 GetCategories() override { return EAssetTypeCategories::Sounds; }
+};
+
+/** Base Settings factories */
+
+class FAssetTypeActions_SoundfieldEncodingSettings : public FAssetTypeActions_Base
+{
+public:
+	// IAssetTypeActions Implementation
+	virtual FText GetName() const override { return NSLOCTEXT("AssetTypeActions", "FAssetTypeActions_SoundfieldEncodingSettings", "Soundfield Encoding Settings"); }
+	virtual FColor GetTypeColor() const override { return Audio::SoundfieldDefaultSubmixColor; }
+	virtual UClass* GetSupportedClass() const override;
+
+	virtual uint32 GetCategories() override { return EAssetTypeCategories::Sounds; }
+};
+
+class FAssetTypeActions_SoundfieldEffectSettings : public FAssetTypeActions_Base
+{
+public:
+	// IAssetTypeActions Implementation
+	virtual FText GetName() const override { return NSLOCTEXT("AssetTypeActions", "FAssetTypeActions_SoundfieldEffectSettings", "Soundfield Effect Settings"); }
+	virtual FColor GetTypeColor() const override { return Audio::SoundfieldDefaultSubmixColor; }
+	virtual UClass* GetSupportedClass() const override;
+
+	virtual uint32 GetCategories() override { return EAssetTypeCategories::Sounds; }
+};
+
+class FAssetTypeActions_SoundfieldEffect : public FAssetTypeActions_Base
+{
+public:
+	// IAssetTypeActions Implementation
+	virtual FText GetName() const override { return NSLOCTEXT("AssetTypeActions", "FAssetTypeActions_SoundfieldEffect", "Soundfield Effect"); }
+	virtual FColor GetTypeColor() const override { return Audio::SoundfieldDefaultSubmixColor; }
+	virtual UClass* GetSupportedClass() const override;
+
+	virtual uint32 GetCategories() override { return EAssetTypeCategories::Sounds; }
+};
+
+class FAssetTypeActions_AudioEndpointSettings : public FAssetTypeActions_Base
+{
+public:
+	// IAssetTypeActions Implementation
+	virtual FText GetName() const override { return NSLOCTEXT("AssetTypeActions", "FAssetTypeActions_AudioEndpointSettings", "Audio Endpoint Settings"); }
+	virtual FColor GetTypeColor() const override { return Audio::EndpointDefaultSubmixColor; }
+	virtual UClass* GetSupportedClass() const override;
+
+	virtual uint32 GetCategories() override { return EAssetTypeCategories::Sounds; }
+};
+
+class FAssetTypeActions_SoundfieldEndpointSettings : public FAssetTypeActions_Base
+{
+public:
+	// IAssetTypeActions Implementation
+	virtual FText GetName() const override { return NSLOCTEXT("AssetTypeActions", "FAssetTypeActions_SoundfieldEndpointSettings", "Soundfield Endpoint Settings"); }
+	virtual FColor GetTypeColor() const override { return Audio::EndpointDefaultSubmixColor; }
+	virtual UClass* GetSupportedClass() const override;
+
+	virtual uint32 GetCategories() override { return EAssetTypeCategories::Sounds; }
+};
diff --git a/Engine/Source/Editor/AudioEditor/Private/AudioEditorModule.cpp b/Engine/Source/Editor/AudioEditor/Private/AudioEditorModule.cpp
--- a/Engine/Source/Editor/AudioEditor/Private/AudioEditorModule.cpp
+++ b/Engine/Source/Editor/AudioEditor/Private/AudioEditorModule.cpp
@@ -149,6 +149,14 @@ public:
 		{
 			IAssetTools& AssetTools = FModuleManager::LoadModuleChecked<FAssetToolsModule>("AssetTools").Get();
 			AssetTools.RegisterAssetTypeActions(MakeShareable(new FAssetTypeActions_SoundSubmix));
+			AssetTools.RegisterAssetTypeActions(MakeShareable(new FAssetTypeActions_SoundfieldSubmix));
+			AssetTools.RegisterAssetTypeActions(MakeShareable(new FAssetTypeActions_EndpointSubmix));
+			AssetTools.RegisterAssetTypeActions(MakeShareable(new FAssetTypeActions_SoundfieldEndpointSubmix));
+			AssetTools.RegisterAssetTypeActions(MakeShareable(new FAssetTypeActions_SoundfieldEncodingSettings));
+			AssetTools.RegisterAssetTypeActions(MakeShareable(new FAssetTypeActions_SoundfieldEffectSettings));
+			AssetTools.RegisterAssetTypeActions(MakeShareable(new FAssetTypeActions_SoundfieldEffect));
+			AssetTools.RegisterAssetTypeActions(MakeShareable(new FAssetTypeActions_AudioEndpointSettings));
+			AssetTools.RegisterAssetTypeActions(MakeShareable(new FAssetTypeActions_SoundfieldEndpointSettings));
 			AssetTools.RegisterAssetTypeActions(MakeShareable(new FAssetTypeActions_SoundEffectSubmixPreset));
 			AssetTools.RegisterAssetTypeActions(MakeShareable(new FAssetTypeActions_SoundEffectSourcePreset));
 			AssetTools.RegisterAssetTypeActions(MakeShareable(new FAssetTypeActions_SoundEffectSourcePresetChain));
@@ -195,7 +203,7 @@ public:
 		return NewSoundClassEditor;
 	}
 
-	virtual TSharedRef<FAssetEditorToolkit> CreateSoundSubmixEditor(const EToolkitMode::Type Mode, const TSharedPtr< IToolkitHost >& InitToolkitHost, USoundSubmix* InSoundSubmix) override
+	virtual TSharedRef<FAssetEditorToolkit> CreateSoundSubmixEditor(const EToolkitMode::Type Mode, const TSharedPtr< IToolkitHost >& InitToolkitHost, USoundSubmixBase* InSoundSubmix) override
 	{
 		TSharedPtr<FSoundSubmixEditor> NewSubmixEditor = MakeShared<FSoundSubmixEditor>();
 		NewSubmixEditor->Init(Mode, InitToolkitHost, InSoundSubmix);
diff --git a/Engine/Source/Editor/AudioEditor/Private/Factories/SoundSubmixFactory.cpp b/Engine/Source/Editor/AudioEditor/Private/Factories/SoundSubmixFactory.cpp
--- a/Engine/Source/Editor/AudioEditor/Private/Factories/SoundSubmixFactory.cpp
+++ b/Engine/Source/Editor/AudioEditor/Private/Factories/SoundSubmixFactory.cpp
@@ -32,3 +32,91 @@ bool USoundSubmixFactory::CanCreateNew() const
 {
 	return GetDefault<UAudioSettings>()->IsAudioMixerEnabled();
 }
+
+// Soundfield Submix Factory: 
+
+USoundfieldSubmixFactory::USoundfieldSubmixFactory(const FObjectInitializer& ObjectInitializer)
+	: Super(ObjectInitializer)
+{
+	SupportedClass = USoundfieldSubmix::StaticClass();
+	bCreateNew = true;
+	bEditorImport = false;
+	bEditAfterNew = true;
+}
+
+UObject* USoundfieldSubmixFactory::FactoryCreateNew(UClass* Class, UObject* InParent, FName Name, EObjectFlags Flags, UObject* Context, FFeedbackContext* Warn)
+{
+	USoundfieldSubmix* SoundSubmix = NewObject<USoundfieldSubmix>(InParent, Name, Flags);
+
+	class FAudioDeviceManager* AudioDeviceManager = GEngine ? GEngine->GetAudioDeviceManager() : nullptr;
+	if (AudioDeviceManager)
+	{
+		AudioDeviceManager->InitSoundSubmixes();
+	}
+
+	return SoundSubmix;
+}
+
+bool USoundfieldSubmixFactory::CanCreateNew() const
+{
+	return GetDefault<UAudioSettings>()->IsAudioMixerEnabled();
+}
+
+// Endpoint Submix Factory:
+
+UEndpointSubmixFactory::UEndpointSubmixFactory(const FObjectInitializer& ObjectInitializer)
+	: Super(ObjectInitializer)
+{
+	SupportedClass = UEndpointSubmix::StaticClass();
+	bCreateNew = true;
+	bEditorImport = false;
+	bEditAfterNew = true;
+}
+
+UObject* UEndpointSubmixFactory::FactoryCreateNew(UClass* Class, UObject* InParent, FName Name, EObjectFlags Flags, UObject* Context, FFeedbackContext* Warn)
+{
+	UEndpointSubmix* SoundSubmix = NewObject<UEndpointSubmix>(InParent, Name, Flags);
+
+	class FAudioDeviceManager* AudioDeviceManager = GEngine ? GEngine->GetAudioDeviceManager() : nullptr;
+	if (AudioDeviceManager)
+	{
+		AudioDeviceManager->InitSoundSubmixes();
+	}
+
+	return SoundSubmix;
+}
+
+bool UEndpointSubmixFactory::CanCreateNew() const
+{
+	return GetDefault<UAudioSettings>()->IsAudioMixerEnabled();
+}
+
+
+// Soundfield Endpoint Submix Factory:
+
+USoundfieldEndpointSubmixFactory::USoundfieldEndpointSubmixFactory(const FObjectInitializer& ObjectInitializer)
+	: Super(ObjectInitializer)
+{
+	SupportedClass = USoundfieldEndpointSubmix::StaticClass();
+	bCreateNew = true;
+	bEditorImport = false;
+	bEditAfterNew = true;
+}
+
+UObject* USoundfieldEndpointSubmixFactory::FactoryCreateNew(UClass* Class, UObject* InParent, FName Name, EObjectFlags Flags, UObject* Context, FFeedbackContext* Warn)
+{
+	USoundfieldEndpointSubmix* SoundSubmix = NewObject<USoundfieldEndpointSubmix>(InParent, Name, Flags);
+
+	class FAudioDeviceManager* AudioDeviceManager = GEngine ? GEngine->GetAudioDeviceManager() : nullptr;
+	if (AudioDeviceManager)
+	{
+		AudioDeviceManager->InitSoundSubmixes();
+	}
+
+	return SoundSubmix;
+}
+
+bool USoundfieldEndpointSubmixFactory::CanCreateNew() const
+{
+	return GetDefault<UAudioSettings>()->IsAudioMixerEnabled();
+}
diff --git a/Engine/Source/Editor/AudioEditor/Private/SoundSubmixEditor.cpp b/Engine/Source/Editor/AudioEditor/Private/SoundSubmixEditor.cpp
--- a/Engine/Source/Editor/AudioEditor/Private/SoundSubmixEditor.cpp
+++ b/Engine/Source/Editor/AudioEditor/Private/SoundSubmixEditor.cpp
@@ -93,7 +93,7 @@ void FSoundSubmixEditor::UnregisterTabSpawners(const TSharedRef<FTabManager>& In
 	InTabManager->UnregisterTabSpawner( PropertiesTabId );
 }
 
-void FSoundSubmixEditor::AddEditableSubmixChildren(USoundSubmix* RootSubmix)
+void FSoundSubmixEditor::AddEditableSubmixChildren(USoundSubmixBase* RootSubmix)
 {
 	if (!RootSubmix)
 	{
@@ -101,7 +101,7 @@ void FSoundSubmixEditor::AddEditableSubmixChildren(USoundSubmix* RootSubmix)
 	}
 
 	RootSubmix->SetFlags(RF_Transactional);
-	for (USoundSubmix* Child : RootSubmix->ChildSubmixes)
+	for (USoundSubmixBase* Child : RootSubmix->ChildSubmixes)
 	{
 		if (Child)
 		{
@@ -114,11 +114,13 @@ void FSoundSubmixEditor::AddEditableSubmixChildren(USoundSubmix* RootSubmix)
 
 void FSoundSubmixEditor::Init(const EToolkitMode::Type Mode, const TSharedPtr<IToolkitHost >& InitToolkitHost, UObject* ObjectToEdit)
 {
-	USoundSubmix* SoundSubmix = CastChecked<USoundSubmix>(ObjectToEdit);
+	USoundSubmixBase* SoundSubmix = CastChecked<USoundSubmixBase>(ObjectToEdit);
 
-	while (SoundSubmix->ParentSubmix)
+	USoundSubmixWithParentBase* SubmixWithParent = Cast<USoundSubmixWithParentBase>(SoundSubmix);
+	while (SubmixWithParent && SubmixWithParent->ParentSubmix)
 	{
-		SoundSubmix = SoundSubmix->ParentSubmix;
+		SoundSubmix = SubmixWithParent->ParentSubmix;
+		SubmixWithParent = Cast<USoundSubmixWithParentBase>(SoundSubmix);
 	}
 
 	GEditor->RegisterForUndo(this);
@@ -259,7 +261,7 @@ FLinearColor FSoundSubmixEditor::GetWorldCentricTabColorScale() const
 	return FLinearColor( 0.2f, 0.4f, 0.8f, 0.5f );
 }
 
-void FSoundSubmixEditor::CreateInternalWidgets(USoundSubmix* InSoundSubmix)
+void FSoundSubmixEditor::CreateInternalWidgets(USoundSubmixBase* InSoundSubmix)
 {
 	GraphEditor = CreateGraphEditorWidget(InSoundSubmix);
 
@@ -269,7 +271,7 @@ void FSoundSubmixEditor::CreateInternalWidgets(USoundSubmix* InSoundSubmix)
 	DetailsView->SetObject(InSoundSubmix);
 }
 
-TSharedRef<SGraphEditor> FSoundSubmixEditor::CreateGraphEditorWidget(USoundSubmix* InSoundSubmix)
+TSharedRef<SGraphEditor> FSoundSubmixEditor::CreateGraphEditorWidget(USoundSubmixBase* InSoundSubmix)
 {
 	if (!GraphEditorCommands.IsValid())
 	{
@@ -339,7 +341,7 @@ void FSoundSubmixEditor::SelectAllNodes()
 	GraphEditor->SelectAllNodes();
 }
 
-void FSoundSubmixEditor::SelectSubmixes(TSet<USoundSubmix*>& InSubmixes)
+void FSoundSubmixEditor::SelectSubmixes(TSet<USoundSubmixBase*>& InSubmixes)
 {
 	TArray<UObject*> ObjectsToSelect;
 
@@ -352,7 +354,7 @@ void FSoundSubmixEditor::SelectSubmixes(TSet<USoundSubmix*>& InSubmixes)
 		if (SelectedNode)
 		{
 			USoundSubmixGraphNode* GraphNode = CastChecked<USoundSubmixGraphNode>(SelectedNode);
-			if (USoundSubmix* Submix = GraphNode->SoundSubmix)
+			if (USoundSubmixBase* Submix = GraphNode->SoundSubmix)
 			{
 				if (InSubmixes.Contains(Submix))
 				{
@@ -428,7 +430,7 @@ void FSoundSubmixEditor::CreateSoundSubmix(UEdGraphPin* FromPin, const FVector2D
 	}
 
 	// Derive new package path from existing asset's path
-	USoundSubmix* SoundSubmix = CastChecked<USoundSubmix>(GetEditingObjects()[0]);
+	USoundSubmixBase* SoundSubmix = CastChecked<USoundSubmixBase>(GetEditingObjects()[0]);
 	FString PackagePath = SoundSubmix->GetPathName();
 	FString AssetName = FString::Printf(TEXT("/%s.%s"), *SoundSubmix->GetName(), *SoundSubmix->GetName());
 	PackagePath.RemoveFromEnd(AssetName);
@@ -474,7 +476,7 @@ void FSoundSubmixEditor::AddMissingEditableSubmixes()
 			USoundSubmix* UntrackedSubmix = GraphNode->SoundSubmix;
 			if (UntrackedSubmix && !GetEditingObjects().Contains(UntrackedSubmix))
 			{
-					bChanged = true;
+				bChanged = true;
 				AddEditingObject(UntrackedSubmix);
 			}
 		}
diff --git a/Engine/Source/Editor/AudioEditor/Private/SoundSubmixEditor.h b/Engine/Source/Editor/AudioEditor/Private/SoundSubmixEditor.h
--- a/Engine/Source/Editor/AudioEditor/Private/SoundSubmixEditor.h
+++ b/Engine/Source/Editor/AudioEditor/Private/SoundSubmixEditor.h
@@ -13,7 +13,7 @@
 
 // Forward Declarations
 class UEdGraphPin;
-class USoundSubmix;
+class USoundSubmixBase;
 class UEdGraph;
 
 class FSoundSubmixEditor : public ISoundSubmixEditor, public FGCObject, public FEditorUndoClient
@@ -57,7 +57,7 @@ public:
 	void AddMissingEditableSubmixes();
 
 	/** Select node associated with the provided submix */
-	void SelectSubmixes(TSet<USoundSubmix*>& InSubmixes);
+	void SelectSubmixes(TSet<USoundSubmixBase*>& InSubmixes);
 
 	/** Returns current graph handled by editor */
 	UEdGraph* GetGraph();
@@ -67,10 +67,10 @@ private:
 	TSharedRef<SDockTab> SpawnTab_Properties(const FSpawnTabArgs& Args);
 
 	/** Creates all internal widgets for the tabs to point at */
-	void CreateInternalWidgets(USoundSubmix* InSoundSubmix);
+	void CreateInternalWidgets(USoundSubmixBase* InSoundSubmix);
 
 	/** Create new graph editor widget */
-	TSharedRef<SGraphEditor> CreateGraphEditorWidget(USoundSubmix* InSoundSubmix);
+	TSharedRef<SGraphEditor> CreateGraphEditorWidget(USoundSubmixBase* InSoundSubmix);
 
 	/** Called when the selection changes in the GraphEditor */
 	void OnSelectedNodesChanged(const TSet<UObject*>& NewSelection);
@@ -79,7 +79,7 @@ private:
 	FActionMenuContent OnCreateGraphActionMenu(UEdGraph* InGraph, const FVector2D& InNodePosition, const TArray<UEdGraphPin*>& InDraggedPins, bool bAutoExpand, SGraphEditor::FActionMenuClosed InOnMenuClosed);
 
 	/** Adds all children of provided root submix as editable */
-	void AddEditableSubmixChildren(USoundSubmix* RootSubmix);
+	void AddEditableSubmixChildren(USoundSubmixBase* RootSubmix);
 
 	/** Select every node in the graph */
 	void SelectAllNodes();
diff --git a/Engine/Source/Editor/AudioEditor/Private/SoundSubmixEditorUtilities.cpp b/Engine/Source/Editor/AudioEditor/Private/SoundSubmixEditorUtilities.cpp
--- a/Engine/Source/Editor/AudioEditor/Private/SoundSubmixEditorUtilities.cpp
+++ b/Engine/Source/Editor/AudioEditor/Private/SoundSubmixEditorUtilities.cpp
@@ -12,7 +12,7 @@ void FSoundSubmixEditorUtilities::CreateSoundSubmix(const UEdGraph* Graph, UEdGr
 	check(Graph);
 
 	// Cast outer to SoundSubmix
-	USoundSubmix* SoundSubmix = CastChecked<USoundSubmix>(Graph->GetOuter());
+	USoundSubmixBase* SoundSubmix = CastChecked<USoundSubmixBase>(Graph->GetOuter());
 
 	if (SoundSubmix != nullptr)
 	{
diff --git a/Engine/Source/Editor/AudioEditor/Private/SoundSubmixGraph.cpp b/Engine/Source/Editor/AudioEditor/Private/SoundSubmixGraph.cpp
--- a/Engine/Source/Editor/AudioEditor/Private/SoundSubmixGraph.cpp
+++ b/Engine/Source/Editor/AudioEditor/Private/SoundSubmixGraph.cpp
@@ -13,13 +13,13 @@ USoundSubmixGraph::USoundSubmixGraph(const FObjectInitializer& ObjectInitializer
 {
 }
 
-void USoundSubmixGraph::SetRootSoundSubmix(USoundSubmix* InSoundSubmix)
+void USoundSubmixGraph::SetRootSoundSubmix(USoundSubmixBase* InSoundSubmix)
 {
 	check(!RootSoundSubmix);
 	RootSoundSubmix = InSoundSubmix;
 }
 
-USoundSubmix* USoundSubmixGraph::GetRootSoundSubmix() const
+USoundSubmixBase* USoundSubmixGraph::GetRootSoundSubmix() const
 {
 	return RootSoundSubmix;
 }
@@ -44,11 +44,11 @@ void USoundSubmixGraph::RebuildGraph()
 	Package->SetDirtyFlag(bIsDirty);
 }
 
-void USoundSubmixGraph::AddDroppedSoundSubmixes(const TSet<USoundSubmix*>& SoundSubmixes, int32 NodePosX, int32 NodePosY)
+void USoundSubmixGraph::AddDroppedSoundSubmixes(const TSet<USoundSubmixBase*>& SoundSubmixes, int32 NodePosX, int32 NodePosY)
 {
 	Modify();
 
-	for (USoundSubmix* SoundSubmix : SoundSubmixes)
+	for (USoundSubmixBase* SoundSubmix : SoundSubmixes)
 	{
 		NodePosY += ConstructNodes(SoundSubmix, NodePosX, NodePosY);
 	}
@@ -56,7 +56,7 @@ void USoundSubmixGraph::AddDroppedSoundSubmixes(const TSet<USoundSubmix*>& Sound
 	NotifyGraphChanged();
 }
 
-void USoundSubmixGraph::AddNewSoundSubmix(UEdGraphPin* FromPin, USoundSubmix* SoundSubmix, int32 NodePosX, int32 NodePosY, bool bSelectNewNode/* = true*/)
+void USoundSubmixGraph::AddNewSoundSubmix(UEdGraphPin* FromPin, USoundSubmixBase* SoundSubmix, int32 NodePosX, int32 NodePosY, bool bSelectNewNode/* = true*/)
 {
 	check(SoundSubmix->ChildSubmixes.Num() == 0);
 
@@ -68,7 +68,7 @@ void USoundSubmixGraph::AddNewSoundSubmix(UEdGraphPin* FromPin, USoundSubmix* So
 	NotifyGraphChanged();
 }
 
-bool USoundSubmixGraph::IsSubmixDisplayed(USoundSubmix* SoundSubmix) const
+bool USoundSubmixGraph::IsSubmixDisplayed(USoundSubmixBase* SoundSubmix) const
 {
 	return FindExistingNode(SoundSubmix) != nullptr;
 }
@@ -84,12 +84,15 @@ void USoundSubmixGraph::LinkSoundSubmixes()
 			Node->SoundSubmix->Modify();
 
 			// remove parents of existing children
-			for (USoundSubmix* ChildSubmix : Node->SoundSubmix->ChildSubmixes)
+			for (USoundSubmixBase* ChildSubmix : Node->SoundSubmix->ChildSubmixes)
 			{
 				if (ChildSubmix)
 				{
 					ChildSubmix->Modify();
-					ChildSubmix->ParentSubmix = nullptr;
+					if (USoundSubmixWithParentBase* SubmixWithParent = CastChecked<USoundSubmixWithParentBase>(ChildSubmix))
+					{
+						SubmixWithParent = nullptr;
+					}
 				}
 			}
 
@@ -99,14 +102,23 @@ void USoundSubmixGraph::LinkSoundSubmixes()
 			{
 				for (UEdGraphPin* GraphPin : ChildPin->LinkedTo)
 				{
+					
+
 					if (!GraphPin)
 					{
 						continue;
 					}
 
 					USoundSubmixGraphNode* ChildNode = CastChecked<USoundSubmixGraphNode>(GraphPin->GetOwningNode());
-					Node->SoundSubmix->ChildSubmixes.Add(ChildNode->SoundSubmix);
-					ChildNode->SoundSubmix->SetParentSubmix(Node->SoundSubmix);
+
+					// If the child submix we're connecting to isn't the type of submix that has an output, continue.
+					USoundSubmixWithParentBase* ChildSubmixWithParent = Cast<USoundSubmixWithParentBase>(ChildNode->SoundSubmix);
+
+					if (ChildSubmixWithParent)
+					{
+						Node->SoundSubmix->ChildSubmixes.Add(ChildNode->SoundSubmix);
+						ChildSubmixWithParent->SetParentSubmix(Node->SoundSubmix);
+					}
 				}
 			}
 
@@ -136,7 +148,7 @@ void USoundSubmixGraph::RefreshGraphLinks()
 			{
 				for (int32 ChildIndex = 0; ChildIndex < Node->SoundSubmix->ChildSubmixes.Num(); ChildIndex++)
 				{
-					USoundSubmix* ChildSubmix = Node->SoundSubmix->ChildSubmixes[ChildIndex];
+					USoundSubmixBase* ChildSubmix = Node->SoundSubmix->ChildSubmixes[ChildIndex];
 
 					if (ChildSubmix)
 					{
@@ -178,11 +190,11 @@ void USoundSubmixGraph::RecursivelyRemoveNodes(const TSet<UObject*> NodesToRemov
 	LinkSoundSubmixes();
 }
 
-int32 USoundSubmixGraph::ConstructNodes(USoundSubmix* SoundSubmix, int32 NodePosX, int32 NodePosY, bool bSelectNewNode/* = true*/)
+int32 USoundSubmixGraph::ConstructNodes(USoundSubmixBase* SoundSubmix, int32 NodePosX, int32 NodePosY, bool bSelectNewNode/* = true*/)
 {
 	check(SoundSubmix);
 
-	TMap<USoundSubmix*, int32> ChildCounts;
+	TMap<USoundSubmixBase*, int32> ChildCounts;
 
 	RecursivelyGatherChildCounts(SoundSubmix, ChildCounts);
 
@@ -191,7 +203,7 @@ int32 USoundSubmixGraph::ConstructNodes(USoundSubmix* SoundSubmix, int32 NodePos
 	return RecursivelyConstructChildNodes(GraphNode, ChildCounts);
 }
 
-int32 USoundSubmixGraph::RecursivelyGatherChildCounts(USoundSubmix* ParentSubmix, TMap<USoundSubmix*, int32>& OutChildCounts)
+int32 USoundSubmixGraph::RecursivelyGatherChildCounts(USoundSubmixBase* ParentSubmix, TMap<USoundSubmixBase*, int32>& OutChildCounts)
 {
 	int32 ChildSize = 0;
 
@@ -212,12 +224,12 @@ int32 USoundSubmixGraph::RecursivelyGatherChildCounts(USoundSubmix* ParentSubmix
 	return ChildSize;
 }
 
-int32 USoundSubmixGraph::RecursivelyConstructChildNodes(USoundSubmixGraphNode* ParentNode, const TMap<USoundSubmix*, int32>& InChildCounts, bool bSelectNewNode /* = true*/)
+int32 USoundSubmixGraph::RecursivelyConstructChildNodes(USoundSubmixGraphNode* ParentNode, const TMap<USoundSubmixBase*, int32>& InChildCounts, bool bSelectNewNode /* = true*/)
 {
 	static const int32 HorizontalSpacing = -400;
 	static const int32 VerticalSpacing = 100;
 
-	USoundSubmix* ParentSubmix = ParentNode->SoundSubmix;
+	USoundSubmixBase* ParentSubmix = ParentNode->SoundSubmix;
 	int32 TotalChildSizeY = InChildCounts.FindChecked(ParentSubmix) * VerticalSpacing;
 	int32 NodeStartY = ParentNode->NodePosY - (TotalChildSizeY * 0.5f) + (VerticalSpacing * 0.5f);
 	int32 NodePosX = ParentNode->NodePosX + HorizontalSpacing;
@@ -262,7 +274,7 @@ void USoundSubmixGraph::RemoveAllNodes()
 	}
 }
 
-USoundSubmixGraphNode* USoundSubmixGraph::CreateNode(USoundSubmix* SoundSubmix, int32 NodePosX, int32 NodePosY, bool bSelectNewNode/* = true*/)
+USoundSubmixGraphNode* USoundSubmixGraph::CreateNode(USoundSubmixBase* SoundSubmix, int32 NodePosX, int32 NodePosY, bool bSelectNewNode/* = true*/)
 {
 	USoundSubmixGraphNode* GraphNode = FindExistingNode(SoundSubmix);
 
@@ -278,7 +290,7 @@ USoundSubmixGraphNode* USoundSubmixGraph::CreateNode(USoundSubmix* SoundSubmix,
 	return GraphNode;
 }
 
-USoundSubmixGraphNode* USoundSubmixGraph::FindExistingNode(USoundSubmix* SoundSubmix) const
+USoundSubmixGraphNode* USoundSubmixGraph::FindExistingNode(USoundSubmixBase* SoundSubmix) const
 {
 	for (int32 NodeIndex = 0; NodeIndex < Nodes.Num(); ++NodeIndex)
 	{
diff --git a/Engine/Source/Editor/AudioEditor/Private/SoundSubmixGraphNode.cpp b/Engine/Source/Editor/AudioEditor/Private/SoundSubmixGraphNode.cpp
--- a/Engine/Source/Editor/AudioEditor/Private/SoundSubmixGraphNode.cpp
+++ b/Engine/Source/Editor/AudioEditor/Private/SoundSubmixGraphNode.cpp
@@ -6,6 +6,7 @@
 #include "SoundSubmixGraph/SoundSubmixGraph.h"
 #include "Toolkits/AssetEditorManager.h"
 #include "SoundSubmixEditor.h"
+#include "SoundSubmixDefaultColorPalette.h"
 
 #define LOCTEXT_NAMESPACE "SoundSubmixGraphNode"
 
@@ -56,15 +57,19 @@ bool USoundSubmixGraphNode::CheckRepresentsSoundSubmix()
 
 FLinearColor USoundSubmixGraphNode::GetNodeTitleColor() const
 {
-	return Super::GetNodeTitleColor();
+	return Audio::GetColorForSubmixType(SoundSubmix);
 }
 
 void USoundSubmixGraphNode::AllocateDefaultPins()
 {
 	check(Pins.Num() == 0);
 
-	ChildPin = CreatePin(EGPD_Input, TEXT("SoundSubmix"), *LOCTEXT("SoundSubmixGraphNode_Input", "Input").ToString());
-	ParentPin = CreatePin(EGPD_Output, TEXT("SoundSubmix"), *LOCTEXT("SoundSubmixGraphNode_Output", "Output").ToString());
+	ChildPin = CreatePin(EGPD_Input, Audio::GetNameForSubmixType(SoundSubmix), *LOCTEXT("SoundSubmixGraphNode_Input", "Input").ToString());
+
+	if (USoundSubmixWithParentBase* NonEndpointSubmix = Cast<USoundSubmixWithParentBase>(SoundSubmix))
+	{
+		ParentPin = CreatePin(EGPD_Output, Audio::GetNameForSubmixType(SoundSubmix), *LOCTEXT("SoundSubmixGraphNode_Output", "Output").ToString());
+	}
 }
 
 void USoundSubmixGraphNode::AutowireNewNode(UEdGraphPin* FromPin)
@@ -106,7 +111,7 @@ bool USoundSubmixGraphNode::CanUserDeleteNode() const
 		{
 			if (SoundSubmix->SoundSubmixGraph == Graph)
 			{
-				USoundSubmix* RootSubmix = CastChecked<USoundSubmixGraph>(Graph)->GetRootSoundSubmix();
+				USoundSubmixBase* RootSubmix = CastChecked<USoundSubmixGraph>(Graph)->GetRootSoundSubmix();
 				if (RootSubmix == SoundSubmix)
 				{
 					return false;
diff --git a/Engine/Source/Editor/AudioEditor/Private/SoundSubmixGraphSchema.cpp b/Engine/Source/Editor/AudioEditor/Private/SoundSubmixGraphSchema.cpp
--- a/Engine/Source/Editor/AudioEditor/Private/SoundSubmixGraphSchema.cpp
+++ b/Engine/Source/Editor/AudioEditor/Private/SoundSubmixGraphSchema.cpp
@@ -19,17 +19,10 @@
 #include "SoundSubmixEditorUtilities.h"
 #include "Toolkits/AssetEditorManager.h"
 #include "ToolMenus.h"
-
+#include "SoundSubmixDefaultColorPalette.h"
 
 #define LOCTEXT_NAMESPACE "SoundSubmixSchema"
 
-
-namespace
-{
-	static const FLinearColor SubmixGraphColor = FColor(175, 255, 0);
-} // namespace <>
-
-
 FConnectionDrawingPolicy* FSoundSubmixGraphConnectionDrawingPolicyFactory::CreateConnectionPolicy(
 	const UEdGraphSchema* Schema,
 	int32 InBackLayerID,
@@ -70,6 +63,9 @@ void FSoundSubmixGraphConnectionDrawingPolicy::DetermineWiringStyle(UEdGraphPin*
 
 	bool bExecuted = false;
 
+	USoundSubmixBase* InputSubmix = OutputPin ? CastChecked<USoundSubmixGraphNode>(OutputPin->GetOwningNode())->SoundSubmix : nullptr;
+	USoundSubmixBase* OutputSubmix = InputPin ? CastChecked<USoundSubmixGraphNode>(InputPin->GetOwningNode())->SoundSubmix : nullptr;
+
 	// Run through the predecessors, and on
 	if (FExecPairingMap* PredecessorMap = PredecessorNodes.Find(OutputPin->GetOwningNode()))
 	{
@@ -78,14 +74,14 @@ void FSoundSubmixGraphConnectionDrawingPolicy::DetermineWiringStyle(UEdGraphPin*
 			bExecuted = true;
 
 			OutParams.WireThickness = ActiveWireThickness;
-			OutParams.WireColor = SubmixGraphColor;
-			OutParams.bDrawBubbles = true;
+			OutParams.WireColor = Audio::GetColorForSubmixType(OutputSubmix);
+			OutParams.bDrawBubbles = Audio::IsConnectionPerformingSoundfieldConversion(InputSubmix, OutputSubmix);
 		}
 	}
 
 	if (!bExecuted)
 	{
-		OutParams.WireColor = SubmixGraphColor;
+		OutParams.WireColor = Audio::GetColorForSubmixType(InputSubmix);
 		OutParams.WireThickness = InactiveWireThickness;
 	}
 }
@@ -241,6 +237,85 @@ const FPinConnectionResponse USoundSubmixGraphSchema::CanCreateConnection(const
 		return FPinConnectionResponse(CONNECT_RESPONSE_DISALLOW, LOCTEXT("ConnectionIncompatible", "Directions are not compatible"));
 	}
 
+	// Note- are input pin and output pin swapped here? Am I losing it?
+	USoundSubmixBase* InputSubmix = CastChecked<USoundSubmixGraphNode>(OutputPin->GetOwningNode())->SoundSubmix;
+	USoundSubmixBase* OutputSubmix = CastChecked<USoundSubmixGraphNode>(InputPin->GetOwningNode())->SoundSubmix;
+
+	// Check to see if this is an endpoint submix.
+	if (!InputSubmix->IsA<USoundSubmixWithParentBase>())
+	{
+		return FPinConnectionResponse(CONNECT_RESPONSE_DISALLOW, LOCTEXT("SubmixIsEndpoint", "Submix you are trying to connect from is an endpoint."));
+	}
+
+	// If we're trying to make a connection between two soundfield submixes, ensure that we can transcode between the two.
+	if (InputSubmix->IsA<USoundfieldSubmix>() && (OutputSubmix->IsA<USoundfieldSubmix>() || OutputSubmix->IsA<USoundfieldEndpointSubmix>()))
+	{
+		USoundfieldSubmix* InputSoundfieldSubmix = Cast<USoundfieldSubmix>(InputSubmix);
+		USoundfieldSubmix* OutputSoundfieldSubmix = Cast<USoundfieldSubmix>(OutputSubmix);
+
+		ISoundfieldFactory* InputFactory = InputSoundfieldSubmix->GetSoundfieldFactoryForSubmix();
+		ISoundfieldFactory* OutputFactory = nullptr;
+
+		const USoundfieldEncodingSettingsBase* InputEncodingSettings = InputSoundfieldSubmix->GetEncodingSettings();
+		const USoundfieldEncodingSettingsBase* OutputEncodingSettings = nullptr;
+
+		if (!OutputSoundfieldSubmix)
+		{
+			USoundfieldEndpointSubmix* SoundfieldEndpointSubmixB = CastChecked<USoundfieldEndpointSubmix>(OutputSubmix);
+			OutputFactory = SoundfieldEndpointSubmixB->GetSoundfieldEndpointForSubmix();
+			OutputEncodingSettings = SoundfieldEndpointSubmixB->GetEncodingSettings();
+		}
+		else
+		{
+			OutputFactory = OutputSoundfieldSubmix->GetSoundfieldFactoryForSubmix();
+			OutputEncodingSettings = OutputSoundfieldSubmix->GetEncodingSettings();
+		}
+
+		if (InputFactory && OutputFactory)
+		{
+			if (!InputEncodingSettings)
+			{
+				InputEncodingSettings = InputFactory->GetDefaultEncodingSettings();
+			}
+
+			if (!InputEncodingSettings)
+			{
+				return FPinConnectionResponse(CONNECT_RESPONSE_DISALLOW, LOCTEXT("SoundfieldSubmixSourceIsInvalid", "Submix you are trying to connect from does not specify default settings. Please implement ISoundfieldFactory::GetDefaultEncodingSettings."));
+			}
+
+			TUniquePtr<ISoundfieldEncodingSettingsProxy> InputEncodingSettingsProxy = InputEncodingSettings->GetProxy();
+
+			if (!InputEncodingSettingsProxy)
+			{
+				return FPinConnectionResponse(CONNECT_RESPONSE_DISALLOW, LOCTEXT("SourceSoundfieldEncodingSettingsAreInvalid", "Submix you are trying to connect from failed to generate a proxy of it's settings. Please check USoundfieldEncodingSettingsBase::GetProxy()."));
+			}
+
+
+			if (!OutputEncodingSettings)
+			{
+				OutputEncodingSettings = OutputFactory->GetDefaultEncodingSettings();
+			}
+
+			if (!OutputEncodingSettings)
+			{
+				return FPinConnectionResponse(CONNECT_RESPONSE_DISALLOW, LOCTEXT("SoundfieldSubmixDestIsInvalid", "Submix you are trying to connect to does not specify default settings. Please implement ISoundfieldFactory::GetDefaultEncodingSettings."));
+			}
+
+			TUniquePtr<ISoundfieldEncodingSettingsProxy> OutputEncodingSettingsProxy = OutputEncodingSettings->GetProxy();
+
+			if (!OutputEncodingSettingsProxy)
+			{
+				return FPinConnectionResponse(CONNECT_RESPONSE_DISALLOW, LOCTEXT("DestSoundfieldEncodingSettingsAreInvalid", "Submix you are trying to connect to failed to generate a proxy of it's settings. Please check USoundfieldEncodingSettingsBase::GetProxy()."));
+			}
+
+			const bool bAreSoundfieldsCompatible = InputFactory->CanTranscodeToSoundfieldFormat(OutputFactory->GetSoundfieldFormatName(), *OutputEncodingSettingsProxy) || OutputFactory->CanTranscodeFromSoundfieldFormat(InputFactory->GetSoundfieldFormatName(), *InputEncodingSettingsProxy);
+			if (!bAreSoundfieldsCompatible)
+			{
+				return FPinConnectionResponse(CONNECT_RESPONSE_DISALLOW, LOCTEXT("SoundfieldSubmixesAreIncompatible", "These two submixes have incompatible types."));
+			}
+		}
+	}
+
 	if (ConnectionCausesLoop(InputPin, OutputPin))
 	{
 		return FPinConnectionResponse(CONNECT_RESPONSE_DISALLOW, LOCTEXT("ConnectionLoop", "Connection would cause loop"));
@@ -276,17 +351,20 @@ bool USoundSubmixGraphSchema::TryCreateConnection(UEdGraphPin* PinA, UEdGraphPin
 		USoundSubmixGraph* Graph = CastChecked<USoundSubmixGraph>(PinA->GetOwningNode()->GetGraph());
 		Graph->LinkSoundSubmixes();
 
-		USoundSubmix* SubmixA = CastChecked<USoundSubmixGraphNode>(PinA->GetOwningNode())->SoundSubmix;
-		USoundSubmix* SubmixB = CastChecked<USoundSubmixGraphNode>(PinB->GetOwningNode())->SoundSubmix;
+		USoundSubmixBase* SubmixA = CastChecked<USoundSubmixGraphNode>(PinA->GetOwningNode())->SoundSubmix;
+		USoundSubmixBase* SubmixB = CastChecked<USoundSubmixGraphNode>(PinB->GetOwningNode())->SoundSubmix;
 
 		bool bReopenEditors = false;
 
+		USoundSubmixWithParentBase* SubmixWithParentA = Cast<USoundSubmixWithParentBase>(SubmixA);
+		USoundSubmixWithParentBase* SubmixWithParentB = Cast<USoundSubmixWithParentBase>(SubmixA);
+
 		// If re-basing root, re-open editor.  This will force the root to be the primary edited node
-		if (Graph->GetRootSoundSubmix() == SubmixA && SubmixA->ParentSubmix != nullptr)
+		if (Graph->GetRootSoundSubmix() == SubmixA && SubmixWithParentA && SubmixWithParentA->ParentSubmix != nullptr)
 		{
 			bReopenEditors = true;
 		}
-		else if (Graph->GetRootSoundSubmix() == SubmixB && SubmixB->ParentSubmix != nullptr)
+		else if (Graph->GetRootSoundSubmix() == SubmixB && SubmixWithParentB && SubmixWithParentB->ParentSubmix != nullptr)
 		{
 			bReopenEditors = true;
 		}
@@ -315,7 +393,7 @@ bool USoundSubmixGraphSchema::ShouldHidePinDefaultValue(UEdGraphPin* Pin) const
 
 FLinearColor USoundSubmixGraphSchema::GetPinTypeColor(const FEdGraphPinType& PinType) const
 {
-	return SubmixGraphColor;
+	return Audio::GetColorForSubmixType(PinType.PinCategory);
 }
 
 void USoundSubmixGraphSchema::BreakNodeLinks(UEdGraphNode& TargetNode) const
@@ -334,6 +412,22 @@ void USoundSubmixGraphSchema::BreakPinLinks(UEdGraphPin& TargetPin, bool bSendsN
 	// if this would notify the node then we need to re-link sound classes
 	if (bSendsNodeNotifcation)
 	{
+		if (USoundSubmixGraphNode* GraphNode = Cast<USoundSubmixGraphNode>(TargetPin.GetOwningNode()))
+		{
+			// Iterate through all child submixes
+			USoundSubmixBase* OutputSubmix = GraphNode->SoundSubmix;
+
+			// Note: If we ever support multiple parents for submixes, this will need to be modified.
+			for (USoundSubmixBase* InputSubmix : OutputSubmix->ChildSubmixes)
+			{
+				if (USoundSubmixWithParentBase* SubmixWithParent = Cast<USoundSubmixWithParentBase>(InputSubmix))
+				{
+					SubmixWithParent->ParentSubmix = nullptr;
+					SubmixWithParent->PostEditChange();
+				}
+			}
+		}
+
 		CastChecked<USoundSubmixGraph>(TargetPin.GetOwningNode()->GetGraph())->LinkSoundSubmixes();
 	}
 }
@@ -343,6 +437,25 @@ void USoundSubmixGraphSchema::BreakSinglePinLink(UEdGraphPin* SourcePin, UEdGrap
 	const FScopedTransaction Transaction( NSLOCTEXT("UnrealEd", "GraphEd_BreakSinglePinLink", "Break Pin Link") );
 	Super::BreakSinglePinLink(SourcePin, TargetPin);
 
+	// Compare the directions
+	UEdGraphPin* InputPin = nullptr;
+	UEdGraphPin* OutputPin = nullptr;
+
+	if (!CategorizePinsByDirection(SourcePin, TargetPin, /*out*/ InputPin, /*out*/ OutputPin))
+	{
+		return;
+	}
+
+	// Note- are input pin and output pin swapped here? Am I losing it?
+	USoundSubmixBase* InputSubmix = CastChecked<USoundSubmixGraphNode>(OutputPin->GetOwningNode())->SoundSubmix;
+	USoundSubmixBase* OutputSubmix = CastChecked<USoundSubmixGraphNode>(InputPin->GetOwningNode())->SoundSubmix;
+
+	if (USoundSubmixWithParentBase* SubmixWithParent = CastChecked<USoundSubmixWithParentBase>(InputSubmix))
+	{
+		SubmixWithParent->ParentSubmix = nullptr;
+		SubmixWithParent->PostEditChange();
+	}
+
 	CastChecked<USoundSubmixGraph>(SourcePin->GetOwningNode()->GetGraph())->LinkSoundSubmixes();
 }
 
@@ -353,15 +466,17 @@ void USoundSubmixGraphSchema::DroppedAssetsOnGraph(const TArray<FAssetData>& Ass
 
 	USoundSubmixGraph* SoundSubmixGraph = CastChecked<USoundSubmixGraph>(Graph);
 	TSet<IAssetEditorInstance*> Editors;
-	TSet<USoundSubmix*> UndisplayedSubmixes;
+	TSet<USoundSubmixBase*> UndisplayedSubmixes;
 	for (const FAssetData& Asset : Assets)
 	{
-		// Walk to the root submix
-		if (USoundSubmix* SoundSubmix = Cast<USoundSubmix>(Asset.GetAsset()))
+		if (USoundSubmixBase* SoundSubmix = Cast<USoundSubmixBase>(Asset.GetAsset()))
 		{
-			while (SoundSubmix->ParentSubmix != nullptr)
+			// Walk to the root submix
+			USoundSubmixWithParentBase* SubmixWithParent = Cast<USoundSubmixWithParentBase>(SoundSubmix);
+			while (SubmixWithParent && SubmixWithParent->ParentSubmix != nullptr)
 			{
-				SoundSubmix = SoundSubmix->ParentSubmix;
+				SoundSubmix = SubmixWithParent->ParentSubmix;
+				SubmixWithParent = Cast<USoundSubmixWithParentBase>(SoundSubmix);
 			}
 
 			if (!SoundSubmixGraph->IsSubmixDisplayed(SoundSubmix))
@@ -398,7 +513,7 @@ void USoundSubmixGraphSchema::DroppedAssetsOnGraph(const TArray<FAssetData>& Ass
 		}
 
 		// If editor is this graph's editor, update editable objects and select dropped submixes.
-		if (USoundSubmix* RootSubmix = SoundSubmixGraph->GetRootSoundSubmix())
+		if (USoundSubmixBase* RootSubmix = SoundSubmixGraph->GetRootSoundSubmix())
 		{
 			UAssetEditorSubsystem* EditorSubsystem = GEditor->GetEditorSubsystem<UAssetEditorSubsystem>();
 			if (IAssetEditorInstance* EditorInstance = EditorSubsystem->FindEditorForAsset(RootSubmix, false /* bFocusIfOpen */))
diff --git a/Engine/Source/Editor/AudioEditor/Public/AudioEditorModule.h b/Engine/Source/Editor/AudioEditor/Public/AudioEditorModule.h
--- a/Engine/Source/Editor/AudioEditor/Public/AudioEditorModule.h
+++ b/Engine/Source/Editor/AudioEditor/Public/AudioEditorModule.h
@@ -12,7 +12,7 @@
 class UDialogueWave;
 class USoundClass;
 class USoundCue;
-class USoundSubmix;
+class USoundSubmixBase;
 class USoundNode;
 class USoundWave;
 struct FDialogueContextMapping;
@@ -38,7 +38,7 @@ public:
 	virtual TSharedRef<FAssetEditorToolkit> CreateSoundClassEditor( const EToolkitMode::Type Mode, const TSharedPtr< IToolkitHost >& InitToolkitHost, USoundClass* InSoundClass ) = 0;
 
 	/** Creates a new sound submix editor for a sound submix object. */
-	virtual TSharedRef<FAssetEditorToolkit> CreateSoundSubmixEditor(const EToolkitMode::Type Mode, const TSharedPtr< IToolkitHost >& InitToolkitHost, USoundSubmix* InSoundSubmix) = 0;
+	virtual TSharedRef<FAssetEditorToolkit> CreateSoundSubmixEditor(const EToolkitMode::Type Mode, const TSharedPtr< IToolkitHost >& InitToolkitHost, USoundSubmixBase* InSoundSubmix) = 0;
 
 	/** Returns the menu extensibility manager for the given audio editor type. */
 	virtual TSharedPtr<FExtensibilityManager> GetSoundClassMenuExtensibilityManager() = 0;
diff --git a/Engine/Source/Editor/AudioEditor/Public/SubmixDetailsCustomization.h b/Engine/Source/Editor/AudioEditor/Public/SubmixDetailsCustomization.h
new file mode 100644
--- /dev/null
+++ b/Engine/Source/Editor/AudioEditor/Public/SubmixDetailsCustomization.h
@@ -0,0 +1,81 @@
+// Copyright 1998-2019 Epic Games, Inc. All Rights Reserved.
+#pragma once
+
+#include "CoreMinimal.h"
+#include "Layout/Visibility.h"
+#include "Input/Reply.h"
+#include "IDetailCustomization.h"
+#include "IPropertyTypeCustomization.h"
+#include "PropertyHandle.h"
+
+class IDetailLayoutBuilder;
+class IDetailGroup;
+
+// Utility class to build combo boxes out of arrays of names.
+class FNameSelectorGenerator : public TSharedFromThis<FNameSelectorGenerator>
+{
+
+public:
+	struct FNameSelectorCallbacks
+	{
+		TUniqueFunction<void(FName)> OnNewNameSelected;
+		TUniqueFunction<FName()> GetCurrentlySelectedName;
+		TUniqueFunction<FString()> GetTooltipText;
+	};
+
+	// Use this to generate a combo box widget.
+	TSharedRef<SWidget> MakeNameSelectorWidget(TArray<FName>& InNameArray, FNameSelectorCallbacks&& InCallbacks);
+
+	
+
+protected:
+	void OnSelectionChanged(TSharedPtr<FName> NameItem, ESelectInfo::Type SelectInfo);
+	TSharedRef<SWidget> HandleResponseComboBoxGenerateWidget(TSharedPtr<FName> StringItem);
+	FText GetComboBoxToolTip() const;
+	FText GetComboBoxContent() const;
+
+	TArray<TSharedPtr<FName>> CachedNameArray;
+	FNameSelectorCallbacks CachedCallbacks;
+};
+
+class AUDIOEDITOR_API FSoundfieldSubmixDetailsCustomization : public IDetailCustomization
+{
+public:
+	// Makes a new instance of this detail layout class
+	static TSharedRef<IDetailCustomization> MakeInstance();
+
+	// IDetailCustomization interface
+	virtual void CustomizeDetails(IDetailLayoutBuilder& DetailLayout) override;
+	// End of IDetailCustomization interface
+
+private:
+	TSharedPtr<FNameSelectorGenerator> SoundfieldFormatNameSelectorGenerator;
+};
+
+class AUDIOEDITOR_API FEndpointSubmixDetailsCustomization : public IDetailCustomization, FNameSelectorGenerator
+{
+public:
+	// Makes a new instance of this detail layout class
+	static TSharedRef<IDetailCustomization> MakeInstance();
+
+	// IDetailCustomization interface
+	virtual void CustomizeDetails(IDetailLayoutBuilder& DetailLayout) override;
+	// End of IDetailCustomization interface
+
+private:
+	TSharedPtr<FNameSelectorGenerator> EndpointTypeNameSelectorGenerator;
+};
+
+class AUDIOEDITOR_API FSoundfieldEndpointSubmixDetailsCustomization : public IDetailCustomization, FNameSelectorGenerator
+{
+public:
+	// Makes a new instance of this detail layout class
+	static TSharedRef<IDetailCustomization> MakeInstance();
+
+	// IDetailCustomization interface
+	virtual void CustomizeDetails(IDetailLayoutBuilder& DetailLayout) override;
+	// End of IDetailCustomization interface
+
+private:
+	TSharedPtr<FNameSelectorGenerator> EndpointTypeNameSelectorGenerator;
+};
diff --git a/Engine/Source/Editor/DetailCustomizations/DetailCustomizations.Build.cs b/Engine/Source/Editor/DetailCustomizations/DetailCustomizations.Build.cs
--- a/Engine/Source/Editor/DetailCustomizations/DetailCustomizations.Build.cs
+++ b/Engine/Source/Editor/DetailCustomizations/DetailCustomizations.Build.cs
@@ -12,6 +12,7 @@ public class DetailCustomizations : ModuleRules
 			new string[] {
 				"AppFramework",
 				"Core",
+				"AudioEditor",
 				"CoreUObject",
 				"ApplicationCore",
 				"DesktopWidgets",
diff --git a/Engine/Source/Editor/DetailCustomizations/Private/DetailCustomizations.cpp b/Engine/Source/Editor/DetailCustomizations/Private/DetailCustomizations.cpp
--- a/Engine/Source/Editor/DetailCustomizations/Private/DetailCustomizations.cpp
+++ b/Engine/Source/Editor/DetailCustomizations/Private/DetailCustomizations.cpp
@@ -82,6 +82,7 @@
 #include "CollisionProfileNameCustomization.h"
 #include "DocumentationActorDetails.h"
 #include "SoundBaseDetails.h"
+#include "SubmixDetailsCustomization.h"
 #include "SoundSourceBusDetails.h"
 #include "SoundWaveDetails.h"
 #include "AudioSettingsDetails.h"
@@ -323,6 +324,9 @@ void FDetailCustomizationsModule::RegisterObjectCustomizations()
 
 	RegisterCustomClassLayout("SoundBase", FOnGetDetailCustomizationInstance::CreateStatic(&FSoundBaseDetails::MakeInstance));
 	RegisterCustomClassLayout("SoundSourceBus", FOnGetDetailCustomizationInstance::CreateStatic(&FSoundSourceBusDetails::MakeInstance));
+	RegisterCustomClassLayout("SoundfieldSubmix", FOnGetDetailCustomizationInstance::CreateStatic(&FSoundfieldSubmixDetailsCustomization::MakeInstance));
+	RegisterCustomClassLayout("EndpointSubmix", FOnGetDetailCustomizationInstance::CreateStatic(&FEndpointSubmixDetailsCustomization::MakeInstance));
+	RegisterCustomClassLayout("SoundfieldEndpointSubmix", FOnGetDetailCustomizationInstance::CreateStatic(&FSoundfieldEndpointSubmixDetailsCustomization::MakeInstance));
 	RegisterCustomClassLayout("SoundWave", FOnGetDetailCustomizationInstance::CreateStatic(&FSoundWaveDetails::MakeInstance));
 	RegisterCustomClassLayout("DialogueWave", FOnGetDetailCustomizationInstance::CreateStatic(&FDialogueWaveDetails::MakeInstance));
 	RegisterCustomClassLayout("BodySetup", FOnGetDetailCustomizationInstance::CreateStatic(&FBodySetupDetails::MakeInstance));
diff --git a/Engine/Source/Runtime/AudioExtensions/AudioExtensions.Build.cs b/Engine/Source/Runtime/AudioExtensions/AudioExtensions.Build.cs
new file mode 100644
--- /dev/null
+++ b/Engine/Source/Runtime/AudioExtensions/AudioExtensions.Build.cs
@@ -0,0 +1,28 @@
+// Copyright 1998-2019 Epic Games, Inc. All Rights Reserved.
+
+namespace UnrealBuildTool.Rules
+{
+	public class AudioExtensions : ModuleRules
+	{
+		public AudioExtensions(ReadOnlyTargetRules Target) : base(Target)
+		{
+			PublicDependencyModuleNames.AddRange(
+				new string[]
+				{
+					"Core",
+					"CoreUObject",
+					"SignalProcessing",
+					"AudioMixerCore"
+				}
+			);
+
+            PrivateDependencyModuleNames.AddRange(
+                new string[]
+                {
+					// This is needed for SoundfieldStatics::GetDefaultPositionMap.
+                    //"AudioMixer"
+                }
+            );
+        }
+	}
+}
diff --git a/Engine/Source/Runtime/AudioExtensions/Public/IAudioExtensionPlugin.h b/Engine/Source/Runtime/AudioExtensions/Public/IAudioExtensionPlugin.h
new file mode 100644
--- /dev/null
+++ b/Engine/Source/Runtime/AudioExtensions/Public/IAudioExtensionPlugin.h
@@ -0,0 +1,835 @@
+// Copyright Epic Games, Inc. All Rights Reserved.
+#pragma once
+
+#include "CoreMinimal.h"
+#include "Features/IModularFeature.h"
+#include "ISoundfieldFormat.h"
+#include "Math/Interval.h"
+#include "Modules/ModuleInterface.h"
+#include "Modules/ModuleManager.h"
+#include "UObject/ObjectMacros.h"
+#include "AudioDefines.h"
+
+#if !UE_BUILD_SHIPPING
+class FCommonViewportClient;
+class FCanvas;
+class UFont;
+class FViewport;
+#endif // !UE_BUILD_SHIPPING
+
+#include "IAudioExtensionPlugin.generated.h"
+
+// Forward Declarations
+class FAudioDevice;
+class USoundSubmix;
+
+/**
+* Enumeration of audio plugin types
+*
+*/
+enum class EAudioPlugin : uint8
+{
+	SPATIALIZATION = 0,
+	REVERB = 1,
+	OCCLUSION = 2,
+	MODULATION = 3,
+
+	COUNT = 4
+};
+
+class IAudioSpatialization;
+class IAudioModulation;
+class IAudioOcclusion;
+class IAudioReverb;
+class IAudioPluginListener;
+class FSoundEffectSubmix;
+typedef TSharedPtr<FSoundEffectSubmix, ESPMode::ThreadSafe> FSoundEffectSubmixPtr;
+
+using TAudioSpatializationPtr = TSharedPtr<IAudioSpatialization, ESPMode::ThreadSafe>;
+using TAudioModulationPtr     = TSharedPtr<IAudioModulation, ESPMode::ThreadSafe>;
+using TAudioOcclusionPtr      = TSharedPtr<IAudioOcclusion, ESPMode::ThreadSafe>;
+using TAudioReverbPtr         = TSharedPtr<IAudioReverb, ESPMode::ThreadSafe>;
+using TAudioPluginListenerPtr = TSharedPtr<IAudioPluginListener, ESPMode::ThreadSafe>;
+
+/**
+* FSpatializationParams
+* Struct for retrieving parameters needed for computing spatialization and occlusion plugins.
+*/
+struct FSpatializationParams
+{
+	/** The listener position (is likely at the origin). */
+	FVector ListenerPosition;
+
+	/** The listener orientation. */
+	FQuat ListenerOrientation;
+
+	/** The emitter position relative to listener. */
+	FVector EmitterPosition;
+
+	/** The emitter world position. */
+	FVector EmitterWorldPosition;
+
+	/** The emitter world rotation. */
+	FQuat EmitterWorldRotation;
+
+	/** The emitter world rotation on callback ago. */
+	FQuat LastEmitterWorldRotation;
+
+	/** The left channel position. */
+	FVector LeftChannelPosition;
+
+	/** The right channel position. */
+	FVector RightChannelPosition;
+
+	/** The distance between listener and emitter. */
+	float Distance;
+
+	/** The normalized omni radius, or the radius that will blend a sound to non-3d */
+	float NormalizedOmniRadius;
+
+	/** The time when this spatialization params was built. */
+	double AudioClock;
+
+	FSpatializationParams()
+		: ListenerPosition(FVector::ZeroVector)
+		, ListenerOrientation(FQuat::Identity)
+		, EmitterPosition(FVector::ZeroVector)
+		, EmitterWorldPosition(FVector::ZeroVector)
+		, EmitterWorldRotation(FQuat::Identity)
+		, LastEmitterWorldRotation(FQuat::Identity)
+		, LeftChannelPosition(FVector::ZeroVector)
+		, RightChannelPosition(FVector::ZeroVector)
+		, Distance(0.0f)
+		, NormalizedOmniRadius(0.0f)
+		, AudioClock(0.0)
+	{}
+};
+
+struct FAudioPluginInitializationParams
+{
+	//Maximum number of sources that can play simultaneously.
+	uint32 NumSources;
+
+	//Number of output channels.
+	uint32 NumOutputChannels;
+
+	//Sample rate.
+	uint32 SampleRate;
+
+	//Buffer length used for each callback.
+	uint32 BufferLength;
+
+	//Pointer to audio device owning this audio plugin.
+	//IMPORTANT: This will be deprecated once the AudioMixer
+	//           is taken out of the experimental branch.
+	FAudioDevice* AudioDevicePtr;
+
+	FAudioPluginInitializationParams()
+		: NumSources(0)
+		, NumOutputChannels(0)
+		, SampleRate(0)
+		, BufferLength(0)
+		, AudioDevicePtr(nullptr)
+	{}
+
+};
+
+struct FAudioPluginSourceInputData
+{
+	// The index of the source voice. Guaranteed to be between 0 and the max number of sources rendered.
+	int32 SourceId;
+
+	// The ID of the audio component associated with the wave instance.
+	uint64 AudioComponentId;
+
+	// The audio input buffer
+	Audio::AlignedFloatBuffer* AudioBuffer;
+
+	// Number of channels of the source audio buffer.
+	int32 NumChannels;
+
+	// The listener orientation.
+	FQuat ListenerOrientation;
+
+	// Spatialization parameters.
+	const FSpatializationParams* SpatializationParams;
+};
+
+struct FAudioPluginSourceOutputData
+{
+	// The audio output buffer
+	Audio::AlignedFloatBuffer AudioBuffer;
+};
+
+/** This is a class which should be overridden to provide users with settings to use for individual sounds */
+UCLASS(config = Engine, abstract, editinlinenew, BlueprintType)
+class AUDIOEXTENSIONS_API USpatializationPluginSourceSettingsBase : public UObject
+{
+	GENERATED_BODY()
+};
+
+/************************************************************************/
+/* IAudioPluginFactory                                             */
+/* This interface is inherited by spatialization, reverb and occlusion  */
+/* plugins to describe specifics of a plugin such as platform support,  */
+/* and display names.                                                   */
+/************************************************************************/
+class IAudioPluginFactory
+{
+public:
+	/*
+	* Returns human-readable string representing the display name of this plugin.
+	* This is the name that will be used in settings and .ini files.
+	* If multiple IAudioPlugin implementations are found that return identical strings here,
+	* The first one of these loaded will be used.
+	*
+	* @return FString of the display name of this plugin.
+	*/
+	virtual FString GetDisplayName()
+	{
+		static FString DisplayName = FString(TEXT("Generic Audio Plugin"));
+		return DisplayName;
+	}
+
+	/*
+	* Returns whether this plugin supports use on the specified platform.
+	* @param Platform an enumerated platform (i.e. Windows, Playstation4, etc.)
+	* @return true if this plugin supports use on Platform, false otherwise.
+	*/
+	virtual bool SupportsPlatform(const FString& PlatformName) = 0;
+
+	/*
+	* Returns whether this plugin sends audio to an external renderer.
+	* if this returns true, the audio engine will not mix the result of the audio process callback
+	* from the plugin into the audio output.
+	*
+	* @return true only if the plugin will handle sending audio to the DAC itself.
+	*/
+	virtual bool IsExternalSend()
+	{
+		return false;
+	}
+};
+
+/************************************************************************/
+/* IAudioSpatializationFactory                                          */
+/* Implement this modular feature to make your Spatialization plugin */
+/* visible to the engine.                                               */
+/************************************************************************/
+class IAudioSpatializationFactory : public IAudioPluginFactory, public IModularFeature
+{
+public:
+	/** Virtual destructor */
+	virtual ~IAudioSpatializationFactory()
+	{
+	}
+
+	// IModularFeature
+	static FName GetModularFeatureName()
+	{
+		static FName AudioExtFeatureName = FName(TEXT("AudioSpatializationPlugin"));
+		return AudioExtFeatureName;
+	}
+
+	/* Begin IAudioPluginWithMetadata implementation */
+	virtual FString GetDisplayName() override
+	{
+		static FString DisplayName = FString(TEXT("Generic Audio Spatialization Plugin"));
+		return DisplayName;
+	}
+	/* End IAudioPluginWithMetadata implementation */
+
+	/**
+	* @return the max amount of channels your plugin supports. For example, override this to
+	*         return 2 to support spatializing mono and stereo sources.
+	*/
+	virtual int32 GetMaxSupportedChannels()
+	{
+		return 1;
+	}
+
+	/**
+	* @return a new instance of your spatialization plugin, owned by a shared pointer.
+	*/
+	virtual TAudioSpatializationPtr CreateNewSpatializationPlugin(FAudioDevice* OwningDevice) = 0;
+
+	/** 
+	* @return the UClass type of your settings for spatialization. This allows us to only pass in user settings for your plugin.
+	*/
+	virtual UClass* GetCustomSpatializationSettingsClass() const
+	{
+		return nullptr;
+	}
+};
+
+/**
+* IAudioSpatialization
+*
+* This class represents instances of a plugin that will process spatialization for a stream of audio.
+* Currently used to process a mono-stream through an HRTF spatialization algorithm into a stereo stream.
+* This algorithm contains an audio effect assigned to every VoiceId (playing sound instance). It assumes
+* the effect is updated in the audio engine update loop with new position information.
+*
+*/
+class IAudioSpatialization 
+{
+public:
+	/** Virtual destructor */
+	virtual ~IAudioSpatialization()
+	{
+	}
+
+	/**
+	* Shuts down the audio plugin.
+	*
+	*/
+	virtual void Shutdown()
+	{
+	}
+
+	virtual void OnDeviceShutdown(FAudioDevice* AudioDevice)
+	{
+	}
+
+	/** DEPRECATED: sets the spatialization effect parameters. */
+	virtual void SetSpatializationParameters(uint32 SourceId, const FSpatializationParams& Params)
+	{
+	}
+
+	/** DEPRECATED: Gets the spatialization effect parameters. */
+	virtual void GetSpatializationParameters(uint32 SourceId, FSpatializationParams& OutParams)
+	{
+	}
+	
+	/** DEPRECATED: Initializes the spatialization effect with the given buffer length. */
+	virtual void InitializeSpatializationEffect(uint32 BufferLength)
+	{
+	}
+
+	/** DEPRECATED: Uses the given HRTF algorithm to spatialize a mono audio stream. */
+	virtual void ProcessSpatializationForVoice(uint32 SourceId, float* InSamples, float* OutSamples, const FVector& Position)
+	{
+	}
+
+	/** DEPRECATED: Uses the given HRTF algorithm to spatialize a mono audio stream, assumes the parameters have already been set before processing. */
+	virtual void ProcessSpatializationForVoice(uint32 SourceId, float* InSamples, float* OutSamples)
+	{
+	}
+
+	/** Called when a source is assigned to a voice. */
+	virtual void OnInitSource(const uint32 SourceId, const FName& AudioComponentUserId, USpatializationPluginSourceSettingsBase* InSettings)
+	{
+	}
+
+	/** Called when a source is done playing and is released. */
+	virtual void OnReleaseSource(const uint32 SourceId)
+	{
+	}
+
+	/** Processes audio with the given input and output data structs.*/
+	virtual void ProcessAudio(const FAudioPluginSourceInputData& InputData, FAudioPluginSourceOutputData& OutputData)
+	{
+	}
+
+	/** Called when all sources have finished processing. */
+	virtual void OnAllSourcesProcessed()
+	{
+	}
+
+	/** Returns whether or not the spatialization effect has been initialized */
+	virtual bool IsSpatializationEffectInitialized() const
+	{
+		return false;
+	}
+
+	/** Initializes the spatialization plugin with the given buffer length. */
+	virtual void Initialize(const FAudioPluginInitializationParams InitializationParams)
+	{
+	}
+
+	/** Creates an audio spatialization effect. */
+	virtual bool CreateSpatializationEffect(uint32 SourceId)
+	{
+		return true;
+	}
+
+	/**	Returns the spatialization effect for the given voice id. */
+	virtual void* GetSpatializationEffect(uint32 SourceId)
+	{
+		return nullptr;
+	}
+};
+
+/** This is a class which should be overridden to provide users with settings to use for individual sounds */
+UCLASS(config = Engine, abstract, editinlinenew, BlueprintType)
+class AUDIOEXTENSIONS_API UOcclusionPluginSourceSettingsBase : public UObject
+{
+	GENERATED_BODY()
+};
+
+/************************************************************************/
+/* IAudioOcclusionFactory                                               */
+/*                                                                      */
+/************************************************************************/
+class IAudioOcclusionFactory : public IAudioPluginFactory, public IModularFeature
+{
+public:
+	/** Virtual destructor */
+	virtual ~IAudioOcclusionFactory()
+	{
+	}
+
+	// IModularFeature
+	static FName GetModularFeatureName()
+	{
+		static FName AudioExtFeatureName = FName(TEXT("AudioOcclusionPlugin"));
+		return AudioExtFeatureName;
+	}
+
+	/* Begin IAudioPluginWithMetadata implementation */
+	virtual FString GetDisplayName() override
+	{
+		static FString DisplayName = FString(TEXT("Generic Audio Occlusion Plugin"));
+		return DisplayName;
+	}
+	/* End IAudioPluginWithMetadata implementation */
+
+	virtual TAudioOcclusionPtr CreateNewOcclusionPlugin(FAudioDevice* OwningDevice) = 0;
+
+	/**
+	* @return the UClass type of your settings for occlusion. This allows us to only pass in user settings for your plugin.
+	*/
+	virtual UClass* GetCustomOcclusionSettingsClass() const
+	{
+		return nullptr;
+	}
+};
+
+class IAudioOcclusion
+{
+public:
+	/** Virtual destructor */
+	virtual ~IAudioOcclusion()
+	{
+	}
+
+	/** Initialize the occlusion plugin with the same rate and number of sources. */
+	virtual void Initialize(const FAudioPluginInitializationParams InitializationParams)
+	{
+	}
+
+	/**
+	* Shuts down the audio plugin.
+	*
+	*/
+	virtual void Shutdown()
+	{
+	}
+
+	/** Called when a source is assigned to a voice. */
+	virtual void OnInitSource(const uint32 SourceId, const FName& AudioComponentUserId, const uint32 NumChannels, UOcclusionPluginSourceSettingsBase* InSettings)
+	{
+	}
+
+	/** Called when a source is done playing and is released. */
+	virtual void OnReleaseSource(const uint32 SourceId)
+	{
+	}
+
+	/** Processes audio with the given input and output data structs.*/
+	virtual void ProcessAudio(const FAudioPluginSourceInputData& InputData, FAudioPluginSourceOutputData& OutputData)
+	{
+	}
+};
+
+
+/** Override to provide users with modulation settings custom to individual sounds */
+UCLASS(config = Engine, abstract, editinlinenew, BlueprintType)
+class AUDIOEXTENSIONS_API USoundModulationPluginSourceSettingsBase : public UObject
+{
+	GENERATED_BODY()
+};
+
+/** Collection of settings available on sound objects */
+USTRUCT(BlueprintType)
+struct AUDIOEXTENSIONS_API FSoundModulation
+{
+	GENERATED_USTRUCT_BODY()
+
+public:
+	UPROPERTY(EditAnywhere, BlueprintReadWrite, Category = Modulation)
+	TArray<USoundModulationPluginSourceSettingsBase*> Settings;
+};
+
+/** Interface to sound that is modulateable, allowing for certain specific 
+  * behaviors to be controlled on the sound level by the modulation system.
+  */
+class AUDIOEXTENSIONS_API ISoundModulatable
+{
+public:
+	virtual ~ISoundModulatable() = default;
+
+	/**
+	 * Returns the modulation settings of the sound
+	 */
+	virtual USoundModulationPluginSourceSettingsBase* FindModulationSettings() const = 0;
+
+	/**
+	 * Gets the object definition id of the given playing sound's instance
+	 */
+	virtual uint32 GetObjectId() const = 0;
+
+	/**
+	 * Returns number of actively instances of sound playing (including virtualized instances)
+	 */
+	virtual int32 GetPlayCount() const = 0;
+
+	/**
+	 * Returns whether or not sound is an editor preview sound
+	 */
+	virtual bool IsPreviewSound() const = 0;
+
+	/**
+	 * Stops sound.
+	 */
+	virtual void Stop() = 0;
+};
+
+/************************************************************************/
+/* IAudioModulationFactory                                              */
+/*                                                                      */
+/************************************************************************/
+class IAudioModulationFactory : public IModularFeature
+{
+public:
+	/** Virtual destructor */
+	virtual ~IAudioModulationFactory()
+	{
+	}
+
+	// IModularFeature
+	static FName GetModularFeatureName()
+	{
+		static FName AudioExtFeatureName = FName(TEXT("AudioModulationPlugin"));
+		return AudioExtFeatureName;
+	}
+
+	virtual const FName& GetDisplayName() const = 0;
+
+	virtual TAudioModulationPtr CreateNewModulationPlugin(FAudioDevice* OwningDevice) = 0;
+
+	/**
+	* @return the UClass type of your settings for modulation. This allows us to only pass in user settings for your plugin.
+	*/
+	virtual UClass* GetCustomModulationSettingsClass() const
+	{
+		return nullptr;
+	}
+};
+
+
+/*
+ * Modulatable controls found on each sound instance
+ * processed by the modulation plugin enabled
+ */
+struct AUDIOEXTENSIONS_API FSoundModulationControls
+{
+	float Volume;
+	float Pitch;
+	float Lowpass;
+	float Highpass;
+
+	TMap<FName, float> Controls;
+
+	FSoundModulationControls()
+		: Volume(1.0f)
+		, Pitch(1.0f)
+		, Lowpass(MAX_FILTER_FREQUENCY)
+		, Highpass(MIN_FILTER_FREQUENCY)
+	{
+	}
+};
+
+/** Parameter allowing modulation control override for systems opting in to the Modulation System. */
+USTRUCT(BlueprintType)
+struct AUDIOEXTENSIONS_API FSoundModulationParameter
+{
+	GENERATED_BODY()
+
+public:
+	/** Name of modulation control to drive parameter. Uses value last cached if control is or becomes invalid. */
+	UPROPERTY(EditAnywhere, BlueprintReadWrite, Category = Modulation)
+	FName Control;
+
+	/** Default modulation parameter value. */
+	UPROPERTY(EditAnywhere, BlueprintReadWrite, Category = Modulation)
+	float Value;
+
+private:
+	float ValueMin;
+	float ValueMax;
+
+public:
+	FSoundModulationParameter()
+		: Control(NAME_None)
+		, Value(0.0f)
+		, ValueMin(0.0f)
+		, ValueMax(1.0f)
+	{
+	}
+
+	FSoundModulationParameter(float InValue, float InValueMin = TNumericLimits<float>::Min(), float InValueMax = TNumericLimits<float>::Max())
+		: Control(NAME_None)
+		, ValueMin(InValueMin)
+		, ValueMax(InValueMax)
+	{
+		Value = FMath::Clamp(InValue, ValueMin, ValueMax);
+	}
+
+	FORCEINLINE operator float() const
+	{
+		return Value;
+	}
+
+	FORCEINLINE FSoundModulationParameter& operator= (const FSoundModulationParameter& InParam)
+	{
+		Control = InParam.Control;
+		ValueMin = InParam.ValueMin;
+		ValueMax = InParam.ValueMax;
+		Value = FMath::Clamp(InParam.Value, ValueMin, ValueMax);
+
+		return *this;
+	}
+
+	FORCEINLINE FSoundModulationParameter& operator= (float InValue)
+	{
+		Value = FMath::Clamp(InValue, ValueMin, ValueMax);
+		return *this;
+	}
+
+	FORCEINLINE float operator* (float InValue)
+	{
+		return InValue * Value;
+	}
+
+	FORCEINLINE float operator/ (float InValue)
+	{
+		return InValue / Value;
+	}
+
+	FORCEINLINE float GetValue() const
+	{
+		return Value;
+	}
+
+	FORCEINLINE float GetMinValue() const
+	{
+		return ValueMin;
+	}
+
+	FORCEINLINE float GetMaxValue() const
+	{
+		return ValueMax;
+	}
+
+	FORCEINLINE void SetBounds(float InValueMin, float InValueMax)
+	{
+		ValueMin = InValueMin;
+		ValueMax = InValueMax;
+
+		Value = FMath::Clamp(Value, InValueMin, InValueMax);
+	}
+
+	/* Returns true if updated by control, false if not */
+	bool SetValue(const FSoundModulationControls& InModControls)
+	{
+		if (Control != NAME_None)
+		{
+			if (const float* ControlValue = InModControls.Controls.Find(Control))
+			{
+				Value = FMath::Clamp(*ControlValue, ValueMin, ValueMax);
+				return true;
+			}
+		}
+
+		Value = FMath::Clamp(Value, ValueMin, ValueMax);
+		return false;
+	}
+};
+
+using FSoundModulationControlIndex = uint32;
+
+class IAudioModulation
+{
+public:
+	/** Virtual destructor */
+	virtual ~IAudioModulation() { }
+
+	/** Calculates initial volume to determine if sound is audible using base settings data */
+	virtual float CalculateInitialVolume(const USoundModulationPluginSourceSettingsBase& InSettingsBase) { return 1.0f; }
+
+	/** Initialize the modulation plugin with the same rate and number of sources */
+	virtual void Initialize(const FAudioPluginInitializationParams& InitializationParams) { }
+
+	/** Called when a USoundBase type begins playing a sound */
+	virtual void OnInitSound(ISoundModulatable& Sound, const USoundModulationPluginSourceSettingsBase& Settings) { }
+
+	/** Called when a source is assigned to a voice */
+	virtual void OnInitSource(const uint32 SourceId, const FName& AudioComponentUserId, const uint32 NumChannels, const USoundModulationPluginSourceSettingsBase& Settings) { }
+
+	/** Called when a source is done playing and is released */
+	virtual void OnReleaseSource(const uint32 SourceId) { }
+
+	/** Called when a USoundBase type stops playing any sounds */
+	virtual void OnReleaseSound(ISoundModulatable& Sound) { }
+
+#if !UE_BUILD_SHIPPING
+	/** Request to post help from active plugin (non-shipping builds only) */
+	virtual bool OnPostHelp(FCommonViewportClient* ViewportClient, const TCHAR* Stream) { return false; };
+
+	/** Render stats pertaining to modulation (non-shipping builds only) */
+	virtual int32 OnRenderStat(FViewport* Viewport, FCanvas* Canvas, int32 X, int32 Y, const UFont& Font, const FVector* ViewLocation, const FRotator* ViewRotation) { return Y; }
+
+	/** Toggle showing render stats pertaining to modulation (non-shipping builds only) */
+	virtual bool OnToggleStat(FCommonViewportClient* ViewportClient, const TCHAR* Stream) { return false; }
+#endif //!UE_BUILD_SHIPPING
+
+	/** Processes audio with the given input and output data structs.*/
+	virtual void ProcessAudio(const FAudioPluginSourceInputData& InputData, FAudioPluginSourceOutputData& OutputData) { }
+
+	/** Processes modulated sound controls, returning whether or not controls were modified and an update is pending. */
+	virtual bool ProcessControls(const uint32 SourceId, FSoundModulationControls& Controls) { return false; }
+
+	/** Processes all modulators */
+	virtual void ProcessModulators(const float Elapsed) { }
+};
+
+/** This is a class which should be overridden to provide users with settings to use for individual sounds */
+UCLASS(config = Engine, abstract, editinlinenew, BlueprintType)
+class AUDIOEXTENSIONS_API UReverbPluginSourceSettingsBase : public UObject
+{
+	GENERATED_BODY()
+};
+ 
+class IAudioReverbFactory : public IAudioPluginFactory, public IModularFeature
+{
+public:
+	/** Virtual destructor */
+	virtual ~IAudioReverbFactory()
+	{
+	}
+
+	// IModularFeature
+	static FName GetModularFeatureName()
+	{
+		static FName AudioExtFeatureName = FName(TEXT("AudioReverbPlugin"));
+		return AudioExtFeatureName;
+	}
+
+	/* Begin IAudioPluginWithMetadata implementation */
+	virtual FString GetDisplayName() override
+	{
+		static FString DisplayName = FString(TEXT("Generic Audio Reverb Plugin"));
+		return DisplayName;
+	}
+	/* End IAudioPluginWithMetadata implementation */
+
+	virtual TAudioReverbPtr CreateNewReverbPlugin(FAudioDevice* OwningDevice) = 0;
+
+	/**
+	* @return the UClass type of your settings for reverb. This allows us to only pass in user settings for your plugin.
+	*/
+	virtual UClass* GetCustomReverbSettingsClass() const
+	{
+		return nullptr;
+	}
+};
+
+class IAudioReverb
+{
+public:
+	/** Virtual destructor */
+	virtual ~IAudioReverb()
+	{
+	}
+
+	/** Initialize the reverb plugin with the same rate and number of sources. */
+	virtual void Initialize(const FAudioPluginInitializationParams InitializationParams)
+	{
+	}
+
+	/**
+	* Shuts down the audio plugin.
+	*
+	*/
+	virtual void Shutdown()
+	{
+	}
+
+	virtual void OnDeviceShutdown(FAudioDevice* AudioDevice)
+	{
+	}
+
+	/** Called when a source is assigned to a voice. */
+	virtual void OnInitSource(const uint32 SourceId, const FName& AudioComponentUserId, const uint32 NumChannels, UReverbPluginSourceSettingsBase* InSettings) = 0;
+
+	/** Called when a source is done playing and is released. */
+	virtual void OnReleaseSource(const uint32 SourceId) = 0;
+
+	/** Returns the plugin-managed effect submix instance */
+	virtual FSoundEffectSubmixPtr GetEffectSubmix() = 0;
+
+	/** Returns the plugin-managed effect submix */
+	virtual USoundSubmix* GetSubmix() = 0;
+
+	/** Processes audio with the given input and output data structs.*/
+	virtual void ProcessSourceAudio(const FAudioPluginSourceInputData& InputData, FAudioPluginSourceOutputData& OutputData)
+	{
+	}
+};
+
+/************************************************************************/
+/* IAudioPluginListener                                                 */
+/* Implementations of this interface can receive updates about the      */
+/* audio listener's position in the game world, as well as other data.  */
+/* to use this, register a ListenerObserver to an audio device using    */
+/* FAudioDevice::RegisterPluginListener().                              */
+/************************************************************************/
+class IAudioPluginListener
+{
+public:
+	virtual ~IAudioPluginListener()
+	{
+	}
+
+	virtual void OnDeviceShutdown(FAudioDevice* AudioDevice)
+	{
+	}
+
+	//This function is called when a game world initializes a listener with an audio device this
+	//IAudioPluginListener is registered to. Please note that it is possible to miss this event
+	//if you register this IAudioPluginListener after the listener is initialized.
+	virtual void OnListenerInitialize(FAudioDevice* AudioDevice, UWorld* ListenerWorld)
+	{
+	}
+
+	// This is overridable for any actions a plugin manager may need to do on the game thread.
+	virtual void OnTick(UWorld* InWorld, const int32 ViewportIndex, const FTransform& ListenerTransform, const float InDeltaSeconds)
+	{
+	}
+
+	// This is overridable for any actions a plugin manager may need to do on a level change.
+	virtual void OnWorldChanged(FAudioDevice* AudioDevice, UWorld* InWorld)
+	{
+	}
+
+	// Called when the listener is updated on the given audio device.
+	virtual void OnListenerUpdated(FAudioDevice* AudioDevice, const int32 ViewportIndex, const FTransform& ListenerTransform, const float InDeltaSeconds)
+	{
+	}
+
+	//Called when the listener is shutdown.
+	virtual void OnListenerShutdown(FAudioDevice* AudioDevice)
+	{
+	}
+};
\ No newline at end of file
diff --git a/Engine/Source/Runtime/AudioMixer/AudioMixer.Build.cs b/Engine/Source/Runtime/AudioMixer/AudioMixer.Build.cs
--- a/Engine/Source/Runtime/AudioMixer/AudioMixer.Build.cs
+++ b/Engine/Source/Runtime/AudioMixer/AudioMixer.Build.cs
@@ -33,7 +33,10 @@ namespace UnrealBuildTool.Rules
 					"Engine",
                     "NonRealtimeAudioRenderer",
                     "AudioMixerCore",
-                    "SignalProcessing"
+                    "SignalProcessing",
+					"AudioPlatformConfiguration",
+					"SoundFieldRendering",
+                    "AudioExtensions",
                 }
 			);
 
diff --git a/Engine/Source/Runtime/AudioMixer/Private/AudioMixerBlueprintLibrary.cpp b/Engine/Source/Runtime/AudioMixer/Private/AudioMixerBlueprintLibrary.cpp
--- a/Engine/Source/Runtime/AudioMixer/Private/AudioMixerBlueprintLibrary.cpp
+++ b/Engine/Source/Runtime/AudioMixer/Private/AudioMixerBlueprintLibrary.cpp
@@ -56,6 +56,8 @@ void UAudioMixerBlueprintLibrary::AddMasterSubmixEffect(const UObject* WorldCont
 
 		FSoundEffectSubmixInitData InitData;
 		InitData.SampleRate = MixerDevice->GetSampleRate();
+		InitData.DeviceID = MixerDevice->DeviceID;
+		InitData.PresetSettings = nullptr;
 
 		// Initialize and set the preset immediately
 		SoundEffectSubmix->Init(InitData);
diff --git a/Engine/Source/Runtime/AudioMixer/Private/AudioMixerChannelMaps.cpp b/Engine/Source/Runtime/AudioMixer/Private/AudioMixerChannelMaps.cpp
--- a/Engine/Source/Runtime/AudioMixer/Private/AudioMixerChannelMaps.cpp
+++ b/Engine/Source/Runtime/AudioMixer/Private/AudioMixerChannelMaps.cpp
@@ -263,10 +263,9 @@ namespace Audio
 		return Index;
 	}
 
-	void FMixerDevice::Get2DChannelMap(bool bIsVorbis, const ESubmixChannelFormat InSubmixChannelType, const int32 NumSourceChannels, const bool bIsCenterChannelOnly, Audio::AlignedFloatBuffer& OutChannelMap) const
+	void FMixerDevice::Get2DChannelMap(bool bIsVorbis, const int32 NumSourceChannels, const bool bIsCenterChannelOnly, Audio::AlignedFloatBuffer& OutChannelMap) const
 	{
-		int32 NumOutputChannels = GetNumChannelsForSubmixFormat(InSubmixChannelType);
-		Get2DChannelMap(bIsVorbis, NumSourceChannels, NumOutputChannels, bIsCenterChannelOnly, OutChannelMap);
+		Get2DChannelMap(bIsVorbis, NumSourceChannels, PlatformInfo.NumChannels, bIsCenterChannelOnly, OutChannelMap);
 	}
 
 	void FMixerDevice::Get2DChannelMap(bool bIsVorbis, const int32 NumSourceChannels, const int32 NumOutputChannels, const bool bIsCenterChannelOnly, Audio::AlignedFloatBuffer& OutChannelMap)
@@ -458,36 +457,36 @@ namespace Audio
 		// Now setup the hard-coded values
 		if (NumChannels == 2)
 		{
-			DefaultChannelAzimuthPosition[EAudioMixerChannel::FrontLeft] = { EAudioMixerChannel::FrontLeft, 270 };
-			DefaultChannelAzimuthPosition[EAudioMixerChannel::FrontRight] = { EAudioMixerChannel::FrontRight, 90 };
+			DefaultChannelAzimuthPositions[EAudioMixerChannel::FrontLeft] = { EAudioMixerChannel::FrontLeft, 270 };
+			DefaultChannelAzimuthPositions[EAudioMixerChannel::FrontRight] = { EAudioMixerChannel::FrontRight, 90 };
 		}
 		else
 		{
-			DefaultChannelAzimuthPosition[EAudioMixerChannel::FrontLeft] = { EAudioMixerChannel::FrontLeft, 330 };
-			DefaultChannelAzimuthPosition[EAudioMixerChannel::FrontRight] = { EAudioMixerChannel::FrontRight, 30 };
+			DefaultChannelAzimuthPositions[EAudioMixerChannel::FrontLeft] = { EAudioMixerChannel::FrontLeft, 330 };
+			DefaultChannelAzimuthPositions[EAudioMixerChannel::FrontRight] = { EAudioMixerChannel::FrontRight, 30 };
 		}
 
 		if (bAllowCenterChannel3DPanning)
 		{
 			// Allow center channel for azimuth computations
-			DefaultChannelAzimuthPosition[EAudioMixerChannel::FrontCenter] = { EAudioMixerChannel::FrontCenter, 0 };
+			DefaultChannelAzimuthPositions[EAudioMixerChannel::FrontCenter] = { EAudioMixerChannel::FrontCenter, 0 };
 		}
 		else
 		{
 			// Ignore front center for azimuth computations. 
-			DefaultChannelAzimuthPosition[EAudioMixerChannel::FrontCenter] = { EAudioMixerChannel::FrontCenter, INDEX_NONE };
+			DefaultChannelAzimuthPositions[EAudioMixerChannel::FrontCenter] = { EAudioMixerChannel::FrontCenter, INDEX_NONE };
 		}
 
 		// Always ignore low frequency channel for azimuth computations. 
-		DefaultChannelAzimuthPosition[EAudioMixerChannel::LowFrequency] = { EAudioMixerChannel::LowFrequency, INDEX_NONE };
+		DefaultChannelAzimuthPositions[EAudioMixerChannel::LowFrequency] = { EAudioMixerChannel::LowFrequency, INDEX_NONE };
 
-		DefaultChannelAzimuthPosition[EAudioMixerChannel::BackLeft] = { EAudioMixerChannel::BackLeft, 210 };
-		DefaultChannelAzimuthPosition[EAudioMixerChannel::BackRight] = { EAudioMixerChannel::BackRight, 150 };
-		DefaultChannelAzimuthPosition[EAudioMixerChannel::FrontLeftOfCenter] = { EAudioMixerChannel::FrontLeftOfCenter, 15 };
-		DefaultChannelAzimuthPosition[EAudioMixerChannel::FrontRightOfCenter] = { EAudioMixerChannel::FrontRightOfCenter, 345 };
-		DefaultChannelAzimuthPosition[EAudioMixerChannel::BackCenter] = { EAudioMixerChannel::BackCenter, 180 };
-		DefaultChannelAzimuthPosition[EAudioMixerChannel::SideLeft] = { EAudioMixerChannel::SideLeft, 250 };
-		DefaultChannelAzimuthPosition[EAudioMixerChannel::SideRight] = { EAudioMixerChannel::SideRight, 110 };
+		DefaultChannelAzimuthPositions[EAudioMixerChannel::BackLeft] = { EAudioMixerChannel::BackLeft, 210 };
+		DefaultChannelAzimuthPositions[EAudioMixerChannel::BackRight] = { EAudioMixerChannel::BackRight, 150 };
+		DefaultChannelAzimuthPositions[EAudioMixerChannel::FrontLeftOfCenter] = { EAudioMixerChannel::FrontLeftOfCenter, 15 };
+		DefaultChannelAzimuthPositions[EAudioMixerChannel::FrontRightOfCenter] = { EAudioMixerChannel::FrontRightOfCenter, 345 };
+		DefaultChannelAzimuthPositions[EAudioMixerChannel::BackCenter] = { EAudioMixerChannel::BackCenter, 180 };
+		DefaultChannelAzimuthPositions[EAudioMixerChannel::SideLeft] = { EAudioMixerChannel::SideLeft, 250 };
+		DefaultChannelAzimuthPositions[EAudioMixerChannel::SideRight] = { EAudioMixerChannel::SideRight, 110 };
 
 		// Check any engine ini overrides for these default positions
 		if (NumChannels != 2)
@@ -509,7 +508,7 @@ namespace Audio
 							bool bIsUnique = true;
 							for (int32 ExistingChannelIndex = 0; ExistingChannelIndex < EAudioMixerChannel::MaxSupportedChannel; ++ExistingChannelIndex)
 							{
-								if (DefaultChannelAzimuthPosition[ExistingChannelIndex].Azimuth == AzimuthPositionOverride)
+								if (DefaultChannelAzimuthPositions[ExistingChannelIndex].Azimuth == AzimuthPositionOverride)
 								{
 									bIsUnique = false;
 
@@ -526,7 +525,7 @@ namespace Audio
 
 							if (bIsUnique)
 							{
-								DefaultChannelAzimuthPosition[MixerChannelType].Azimuth = AzimuthPositionOverride;
+								DefaultChannelAzimuthPositions[MixerChannelType].Azimuth = AzimuthPositionOverride;
 							}
 						}
 						else
@@ -547,104 +546,28 @@ namespace Audio
 			}
 		};
 
-		// Build a map of azimuth positions of only the current audio device's output channels
-		ChannelAzimuthPositions.Reset();
+		// Build a array of azimuth positions of only the current audio device's output channels
+		DeviceChannelAzimuthPositions.Reset();
 
 		// Setup the default channel azimuth positions
 		TArray<FChannelPositionInfo> DevicePositions;
 		for (EAudioMixerChannel::Type Channel : PlatformInfo.OutputChannelArray)
 		{
 			// Only track non-LFE and non-Center channel azimuths for use with 3d channel mappings
-			if (Channel != EAudioMixerChannel::LowFrequency && DefaultChannelAzimuthPosition[Channel].Azimuth >= 0)
+			if (Channel != EAudioMixerChannel::LowFrequency && DefaultChannelAzimuthPositions[Channel].Azimuth >= 0)
 			{
-				DevicePositions.Add(DefaultChannelAzimuthPosition[Channel]);
+				DeviceChannelAzimuthPositions.Add(DefaultChannelAzimuthPositions[Channel]);
 			}
 		}
-		DevicePositions.Sort(FCompareByAzimuth());
-		ChannelAzimuthPositions.Add(ESubmixChannelFormat::Device, DevicePositions);
-		OutputChannels[int32(ESubmixChannelFormat::Device)] = 0;
-
-		// Now add channel mappings for the other submix types
-		TArray<FChannelPositionInfo> StereoPositions;
-		StereoPositions.Add(DefaultChannelAzimuthPosition[EAudioMixerChannel::FrontLeft]);
-		StereoPositions.Add(DefaultChannelAzimuthPosition[EAudioMixerChannel::FrontRight]);
-		StereoPositions.Sort(FCompareByAzimuth());
-		ChannelAzimuthPositions.Add(ESubmixChannelFormat::Stereo, StereoPositions);
-		OutputChannels[int32(ESubmixChannelFormat::Stereo)] = StereoPositions.Num();
-
-		TArray<FChannelPositionInfo> QuadPositions;
-		QuadPositions.Add(DefaultChannelAzimuthPosition[EAudioMixerChannel::FrontLeft]);
-		QuadPositions.Add(DefaultChannelAzimuthPosition[EAudioMixerChannel::FrontRight]);
-		QuadPositions.Add(DefaultChannelAzimuthPosition[EAudioMixerChannel::SideLeft]);
-		QuadPositions.Add(DefaultChannelAzimuthPosition[EAudioMixerChannel::SideRight]);
-		QuadPositions.Sort(FCompareByAzimuth());
-		ChannelAzimuthPositions.Add(ESubmixChannelFormat::Quad, QuadPositions);
-		OutputChannels[int32(ESubmixChannelFormat::Quad)] = QuadPositions.Num();
-
-		TArray<FChannelPositionInfo> FiveDotOnePositions;
-		FiveDotOnePositions.Add(DefaultChannelAzimuthPosition[EAudioMixerChannel::FrontLeft]);
-		FiveDotOnePositions.Add(DefaultChannelAzimuthPosition[EAudioMixerChannel::FrontRight]);
-		FiveDotOnePositions.Add(DefaultChannelAzimuthPosition[EAudioMixerChannel::SideLeft]);
-		FiveDotOnePositions.Add(DefaultChannelAzimuthPosition[EAudioMixerChannel::SideRight]);
-		FiveDotOnePositions.Sort(FCompareByAzimuth());
-		ChannelAzimuthPositions.Add(ESubmixChannelFormat::FiveDotOne, FiveDotOnePositions);
-		OutputChannels[int32(ESubmixChannelFormat::FiveDotOne)] = FiveDotOnePositions.Num();
-
-		TArray<FChannelPositionInfo> SevenDotOnePositions;
-		SevenDotOnePositions.Add(DefaultChannelAzimuthPosition[EAudioMixerChannel::FrontLeft]);
-		SevenDotOnePositions.Add(DefaultChannelAzimuthPosition[EAudioMixerChannel::FrontRight]);
-		SevenDotOnePositions.Add(DefaultChannelAzimuthPosition[EAudioMixerChannel::BackLeft]);
-		SevenDotOnePositions.Add(DefaultChannelAzimuthPosition[EAudioMixerChannel::BackRight]);
-		SevenDotOnePositions.Add(DefaultChannelAzimuthPosition[EAudioMixerChannel::SideLeft]);
-		SevenDotOnePositions.Add(DefaultChannelAzimuthPosition[EAudioMixerChannel::SideRight]);
-		SevenDotOnePositions.Sort(FCompareByAzimuth());
-		ChannelAzimuthPositions.Add(ESubmixChannelFormat::SevenDotOne, SevenDotOnePositions);
-		OutputChannels[int32(ESubmixChannelFormat::SevenDotOne)] = SevenDotOnePositions.Num();
-
-		// ambisonics is special cased and uses a plugin.
-		TArray<FChannelPositionInfo> FirstOrderAmbisonicsPositions;
-		FirstOrderAmbisonicsPositions.Add(DefaultChannelAzimuthPosition[EAudioMixerChannel::Ambisonics_W]);
-		FirstOrderAmbisonicsPositions.Add(DefaultChannelAzimuthPosition[EAudioMixerChannel::Ambisonics_X]);
-		FirstOrderAmbisonicsPositions.Add(DefaultChannelAzimuthPosition[EAudioMixerChannel::Ambisonics_Y]);
-		FirstOrderAmbisonicsPositions.Add(DefaultChannelAzimuthPosition[EAudioMixerChannel::Ambisonics_Z]);
-		FirstOrderAmbisonicsPositions.Sort(FCompareByAzimuth());
-		ChannelAzimuthPositions.Add(ESubmixChannelFormat::Ambisonics, FirstOrderAmbisonicsPositions);
-		OutputChannels[int32(ESubmixChannelFormat::Ambisonics)] = FirstOrderAmbisonicsPositions.Num();
+		DeviceChannelAzimuthPositions.Sort(FCompareByAzimuth());
 	}
 
-	int32 FMixerDevice::GetNumChannelsForSubmixFormat(const ESubmixChannelFormat InSubmixChannelType) const
+	const TArray<EAudioMixerChannel::Type>& FMixerDevice::GetChannelArray() const
 	{
-		if (InSubmixChannelType == ESubmixChannelFormat::Device)
-		{
-			return PlatformInfo.NumChannels;
-		}
-
-		return OutputChannels[(int32)InSubmixChannelType];
+		return PlatformInfo.OutputChannelArray;
 	}
-
-	ESubmixChannelFormat FMixerDevice::GetSubmixChannelFormatForNumChannels(const int32 InNumChannels) const
+	const FChannelPositionInfo* FMixerDevice::GetDefaultChannelPositions() const
 	{
-		for (int32 i = 0; i < (int32)ESubmixChannelFormat::Count; ++i)
-		{
-			if (OutputChannels[i] == InNumChannels)
-			{
-				return (ESubmixChannelFormat)i;
-			}
-		}
-		ensureMsgf(false, TEXT("Unsupported number of submix channels %d"), InNumChannels);
-		return ESubmixChannelFormat::Device;
-	}
-
-	const TArray<EAudioMixerChannel::Type>& FMixerDevice::GetChannelArrayForSubmixChannelFormat(const ESubmixChannelFormat InSubmixChannelFormat) const
-	{
-		if (InSubmixChannelFormat == ESubmixChannelFormat::Device)
-		{
-			return PlatformInfo.OutputChannelArray;
-		}
-
-		static_assert(static_cast<int32>(ESubmixChannelFormat::Count) == 6, "ESubmixChannelFormat count must match SubmixOutputChannelMatrix size");
-
-		const int32 ChannelFormatIndex = static_cast<int32>(InSubmixChannelFormat);
-		return SubmixOutputChannelMatrix[ChannelFormatIndex];
+		return DefaultChannelAzimuthPositions;
 	}
 }
