Audio Codecs V1 Hello World and Unit tests

diff --git a/Engine/Source/Runtime/AudioExtensions/Private/AudioExtensionsModule.cpp b/Engine/Source/Runtime/AudioExtensions/Private/AudioExtensionsModule.cpp
--- a/Engine/Source/Runtime/AudioExtensions/Private/AudioExtensionsModule.cpp
+++ b/Engine/Source/Runtime/AudioExtensions/Private/AudioExtensionsModule.cpp
@@ -1,16 +1,20 @@
 // Copyright Epic Games, Inc. All Rights Reserved.
 
-
 #include "Modules/ModuleManager.h"
 
+#include "IAudioCodecRegistry.h"
+#include "PcmCodec.h"
+
 class FAudioExtensionsModule : public IModuleInterface
 {
 public:
-
-	virtual void StartupModule() override
+	void StartupModule() override
 	{
 		FModuleManager::Get().LoadModuleChecked(TEXT("SignalProcessing"));
 		FModuleManager::Get().LoadModuleChecked(TEXT("AudioMixerCore"));
+
+		// Built in Codecs.
+		Audio::ICodecRegistry::Get().RegisterCodec(Audio::Create_PcmCodec());
 	}
 };
 
diff --git a/Engine/Source/Runtime/AudioExtensions/Private/IAudioCodec.cpp b/Engine/Source/Runtime/AudioExtensions/Private/IAudioCodec.cpp
new file mode 100644
--- /dev/null
+++ b/Engine/Source/Runtime/AudioExtensions/Private/IAudioCodec.cpp
@@ -0,0 +1,467 @@
+// Copyright Epic Games, Inc. All Rights Reserved.
+
+#include "IAudioCodec.h"
+#include "Serialization/MemoryWriter.h"
+
+namespace Audio
+{
+	// If < C++17, constexpr requires storage for linkage on Clang.
+	constexpr FFcc FFormatDescriptorSection::kSectionName;
+	constexpr FFcc FHeaderSection::kName;
+	constexpr FFcc FSampleSection::kName;
+	
+	// Encoder INPUT object. (simple for now).
+	template<typename TSampleType>
+	struct TEncoderInputObject : public IEncoderInput
+	{
+		using FSampleType = TSampleType;
+		TSampleBuffer<FSampleType> Samples;
+		FFormat Format;
+
+		TEncoderInputObject(TArrayView<const FSampleType> InSamples, const FFormat& InFormat)
+			: Samples(InSamples.GetData(), InSamples.Num(), InFormat.NumChannels, InFormat.NumFramesPerSec)
+			, Format(InFormat)
+		{}
+		
+		TEncoderInputObject(const TSampleBuffer<int16>& InBuffer, const FFormat& InFormat)
+			: Samples(InBuffer)
+			, Format(InFormat)
+		{}
+
+		const TSampleBuffer<int16>& GetSamples() const override 
+		{ 
+			return Samples; 
+		}
+		FString GetHashForDDC() const override 
+		{ 
+			return FMD5::HashBytes((const uint8*)Samples.GetData(), Samples.GetNumFrames() * sizeof(FSampleType)); 
+		}
+		bool GetFormat(FFormat& OutFormat) const override 
+		{ 
+			OutFormat = Format; 
+			return true; 
+		}
+	};
+	
+	TUniquePtr<IEncoderInput>
+	IEncoderInput::Create(
+		TArrayView<const int16> In16bitInterleave, const FFormat& InFmt )
+	{				
+		if (!audio_ensure(InFmt.NumChannels > 0))
+		{
+			return nullptr;		
+		}
+		if (!audio_ensure(InFmt.NumFramesPerSec > 0))
+		{
+			return nullptr;
+		}
+		return MakeUnique<TEncoderInputObject<int16>>(In16bitInterleave, InFmt);	
+	}
+		
+	struct FEncoderOutputBase : public IEncoderOutput
+	{
+		TUniquePtr<FArchive> Archive;
+		TUniquePtr<FNestedSection> HeaderSection;
+		TUniquePtr<FNestedSection> SampleSection;
+		TSet<FFcc> SectionNames;
+
+		virtual ~FEncoderOutputBase() = default;
+
+		FEncoderOutputBase(TUniquePtr<FArchive>&& InArchive)
+			: Archive(MoveTemp(InArchive))
+		{}
+
+		// Named section "format" for example.
+		bool WriteSection(
+			FEncodedSectionBase& InOutSection) override
+		{
+			// Prevent unnamed sections.
+			if (!audio_ensure(!InOutSection.SectionName.IsNone()))
+			{
+				return false;
+			}
+			
+			// Section names must be unique.
+			if (!audio_ensure(SectionNames.Find(InOutSection.SectionName)==nullptr))
+			{
+				return false;
+			}
+			
+			// Make sure the archive is valid and is in write mode.
+			if (audio_ensure(Archive.IsValid() && Archive->IsSaving()))
+			{
+				InOutSection.BeginSection(*Archive);
+				InOutSection.Serialize(*Archive);
+				InOutSection.EndSection(*Archive);
+				SectionNames.Add(InOutSection.SectionName);
+				return true;
+			}
+			return false;
+		}
+
+		// Header. (header is a bunch of named sections).
+		bool BeginHeader() override
+		{
+			if (audio_ensure(!HeaderSection.IsValid()))
+			{
+				HeaderSection = MakeUnique<FHeaderSection>();
+				HeaderSection->BeginSection(*Archive);
+				return true;
+			}
+			return false;
+		}
+		bool EndHeader() override
+		{
+			if (audio_ensure(HeaderSection.IsValid()))
+			{
+				HeaderSection->EndSection(*Archive);
+				HeaderSection.Reset();
+				return true;
+			}
+			return false;
+		}
+
+		// Sample data.
+		bool BeginSampleData() override
+		{
+			if (audio_ensure(!SampleSection.IsValid()))
+			{
+				SampleSection = MakeUnique<FSampleSection>();
+				SampleSection->BeginSection(*Archive);
+				return true;
+			}
+			return false;
+		}
+		bool EndSampleData() override
+		{
+			if (audio_ensure(SampleSection.IsValid()))
+			{
+				SampleSection->EndSection(*Archive);
+				SampleSection.Reset();
+				return true;
+			}
+			return false;
+		}	
+
+		bool WriteRaw(const void* InRawData, int32 InSize) override
+		{
+			if (!audio_ensure(InRawData) || !audio_ensure(InSize > 0))
+			{
+				return false;
+			}
+			if (!audio_ensure(Archive.IsValid()) || !audio_ensure(!Archive->IsError()) || !audio_ensure(Archive->IsSaving()))
+			{
+				return false;
+			}
+			Archive->Serialize(const_cast<void*>(InRawData),InSize);
+			return Archive->IsError();
+		}
+	};
+	
+	// Encoder Output object. (simple for now).
+	IEncoderOutput::FPtr 
+	IEncoderOutput::Create(
+		TArray<uint8>& OutExternalBuffer)
+	{
+		return MakeUnique<FEncoderOutputBase>(MakeUnique<FMemoryWriter>(OutExternalBuffer));
+	}
+
+	Audio::IEncoderOutput::FPtr IEncoderOutput::Create()
+	{
+		struct FArrayOutputObject : public FEncoderOutputBase
+		{
+			TArray<uint8> Buffer;
+			FArrayOutputObject() : FEncoderOutputBase(MakeUnique<FMemoryWriter>(Buffer)) {}
+		};
+		return MakeUnique<FArrayOutputObject>();
+	}
+
+	TUniquePtr<IDecoderInput> IDecoderInput::Create(
+		TArrayView<const uint8> InBytes, uint32 InOffsetInStream, FName InOldFormatName )
+	{
+		if (InOffsetInStream == 0)
+		{
+			if (!FDecoderInputArrayView::IsValidHeader(InBytes))
+			{
+				// Back compatibility of OLD data, these will be per codec. 
+				return nullptr;
+			}
+		}		
+		return MakeUnique<FDecoderInputArrayView>(InBytes, InOffsetInStream);
+	}
+
+	bool FDecoderInputBase::ParseSections(FArchive& Ar)
+	{		
+		// HEADER section.
+		{
+			// Make sure this looks sane.
+			FHeaderSection Header;
+			uint32 HeaderStart = Ar.Tell();
+			Header.BeginSection(Ar);
+			if (!audio_ensure(!Header.SectionName.IsNone()))
+			{
+				return false;
+			}
+			if (!audio_ensure(Header.SectionName == FHeaderSection::kName))
+			{
+				return false;
+			}
+
+			Toc.Add(Header.SectionName, HeaderStart);
+			if (!audio_ensure(MakeTocNested(Ar, HeaderStart, Header)))
+			{
+				return false;
+			}
+			Header.EndSection(Ar);
+		}
+
+		// SAMPLES section.
+		{		
+			FSampleSection SampleSection;
+			uint32 SamplesStart = Ar.Tell();
+			SampleSection.BeginSection(Ar);
+			if (!audio_ensure(!SampleSection.SectionName.IsNone()))
+			{
+				return false;
+			}
+			if (!audio_ensure(SampleSection.SectionName == FSampleSection::kName))
+			{
+				return false;
+			}
+
+			Toc.Add(SampleSection.SectionName, SamplesStart);
+			SampleSection.EndSection(Ar);
+		}
+			
+		// Looks sane.
+		return true;
+	}
+
+	void FDecoderInputBase::SetArchive(FArchive* InArchive, uint32 InOffsetInStream)
+	{
+		if (audio_ensure(InArchive))
+		{
+			Archive = InArchive;
+			OffsetInStream = InOffsetInStream;
+			if (InOffsetInStream == 0)
+			{
+				bError = !ParseSections(*Archive);
+			}
+			else
+			{
+				Toc.Empty();
+			}
+		}
+	}
+
+	bool FDecoderInputBase::FindSection(FEncodedSectionBase& OutSection)
+	{
+		if (!audio_ensure(Archive->IsLoading()))
+		{
+			return false;
+		}
+
+		if (int32* pPos = Toc.Find(OutSection.SectionName))
+		{
+			Archive->Seek(*pPos);
+			OutSection.BeginSection(*Archive);
+			OutSection.Serialize(*Archive);
+			OutSection.EndSection(*Archive);
+			return true;
+		}
+		return false;
+	}
+
+	TArrayView<const uint8> FDecoderInputBase::PeekNextPacket(
+		int32 InMaxPacketLength) const
+	{		
+		// TODO.
+		check(false);
+		return MakeArrayView(static_cast<const uint8*>(0), 0);
+	}
+
+	TArrayView<const uint8> FDecoderInputBase::PopNextPacket(
+		int32 InPacketSize)
+	{	
+		if (PopBuffer.Num() < InPacketSize)
+		{
+			PopBuffer.SetNumZeroed(InPacketSize);
+		}		
+		if (!bSampleSection )
+		{
+			// To start serving packets, we must be in the sample section.
+			if( !audio_ensure(BeginSampleData()))
+			{
+				// Fail.
+				return MakeArrayView(static_cast<const uint8*>(0), 0);
+			}
+		}
+		if (!audio_ensure(!Archive->IsError()))
+		{
+			return MakeArrayView(static_cast<const uint8*>(0), 0);
+		}
+		Archive->Serialize(PopBuffer.GetData(), InPacketSize);
+		if (!audio_ensure(!Archive->IsError()))
+		{
+			return MakeArrayView(static_cast<const uint8*>(0), 0);
+		}
+		return MakeArrayView(PopBuffer.GetData(),InPacketSize);
+	}
+	
+	bool FDecoderInputBase::BeginSampleData() 
+	{	
+		if (!audio_ensure(!bSampleSection))
+		{
+			return false;
+		}		
+		int32* pSampleDataStart = Toc.Find(FSampleSection::kName);
+		if (!audio_ensure(pSampleDataStart))
+		{
+			return false;
+		}
+		Archive->Seek(*pSampleDataStart);
+		FNestedSection Nested;
+		Nested.BeginSection(*Archive);		
+		if (!audio_ensure(Nested.SectionName == FSampleSection::kName))
+		{
+			return false;
+		}
+		bSampleSection = true;
+		return true;
+	}
+	bool FDecoderInputBase::EndSampleData() 
+	{	
+		bSampleSection = false;
+		return true;
+	}
+
+	bool FDecoderInputBase::MakeTocNested(FArchive &Ar, uint32 HeaderStart, FNestedSection &Header)
+	{
+		while (Ar.Tell() - HeaderStart < Header.SectionSize)
+		{
+			FNestedSection Sect;
+			uint32_t Pos = Ar.Tell();
+			Sect.BeginSection(Ar);
+
+			if (!audio_ensure(!Sect.SectionName.IsNone()) ||
+				!audio_ensure(Toc.Find(Sect.SectionName) == nullptr))
+			{
+				return false;
+			}
+
+			Toc.Add(Sect.SectionName, Pos);
+			Sect.EndSection(Ar);
+		}
+		return true;
+	}
+
+	int64 FDecoderInputBase::Tell() const 
+	{
+		return Archive->Tell() + OffsetInStream;
+	}
+
+	FFormatDescriptorSection::FFormatDescriptorSection(
+		FName InCodecName, 
+		FName InCodecFamilyName, 
+		uint32 InCodecVersion, 
+		uint32 InNumChannels, 
+		uint32 InNumFrames, 
+		uint32 InNumFramesPerSec,
+		uint32 InNumBytesPerPacket ) 
+		: FEncodedSectionBase(kSectionName, kSectionVer) 
+		, CodecName(InCodecName)
+		, CodecFamilyName(InCodecFamilyName)
+		, CodecVersion(InCodecVersion)
+		, NumChannels(InNumChannels)
+		, NumFrames(InNumFrames)
+		, NumFramesPerSec(InNumFramesPerSec)
+		, NumBytesPerPacket(InNumBytesPerPacket)
+	{
+	}
+
+	void FFormatDescriptorSection::Serialize(FArchive& Ar)
+	{
+		Ar << CodecName;
+		Ar << CodecFamilyName;
+		Ar << CodecVersion;
+
+		Ar << NumChannels;
+		Ar << NumFrames;
+		Ar << NumFramesPerSec;
+	}
+
+	bool FEncodedSectionBase::BeginSection(FArchive& Ar)
+	{
+		Ar << SectionName;
+		Ar << SectionVersion;
+
+		// Save where the size is in the stream, so we can return to it, once we know the size of the section.
+		SizePos = Ar.Tell();
+		Ar << SectionSize;
+		return !Ar.IsError();
+	}
+
+	bool FEncodedSectionBase::EndSection(FArchive& Ar)
+	{
+		if (Ar.IsSaving())
+		{
+			uint64 EndPos = Ar.Tell();
+			Ar.Seek(SizePos);
+			SectionSize = (EndPos - SizePos); 
+			Ar << SectionSize;
+			Ar.Seek(EndPos);
+		}
+		else if(Ar.IsLoading())
+		{
+			Ar.Seek(SizePos + SectionSize);
+		}
+		return !Ar.IsError();
+	}
+
+	TUniquePtr<Audio::IDecoderOutput> IDecoderOutput::Create(
+		const FRequirements& Requirements)
+	{
+		if (Requirements.DownstreamFormat == EBitRepresentation::Float32_Interleaved)
+		{
+			return MakeUnique<TDecoderOutputOwnBuffer<float>>(Requirements);
+		}
+		else if (Requirements.DownstreamFormat == EBitRepresentation::Int16_Interleaved)
+		{
+			return MakeUnique<TDecoderOutputOwnBuffer<int16>>(Requirements);
+		}
+		checkNoEntry();
+		return nullptr;
+	}
+
+	TUniquePtr<Audio::IDecoderOutput> IDecoderOutput::Create(
+		const FRequirements& Requirements, 
+		TArrayView<uint8> InExternalMemory)
+	{
+		if (Requirements.DownstreamFormat == EBitRepresentation::Float32_Interleaved)
+		{
+			 auto Output = MakeUnique<TDecoderOutputArrayView<float>>(Requirements);
+			 Output->SetArrayViewBytes(InExternalMemory);
+			 return Output;
+		}
+		else if (Requirements.DownstreamFormat == EBitRepresentation::Int16_Interleaved)
+		{
+			auto Output = MakeUnique<TDecoderOutputArrayView<int16>>(Requirements);
+			Output->SetArrayViewBytes(InExternalMemory);
+			return Output;
+		}
+		checkNoEntry();
+		return nullptr;
+	}
+
+	bool FDecoderInputArrayView::IsValidHeader(
+		TArrayView<const uint8> InBytes)
+	{
+		FMemoryReaderView Ar(InBytes);
+		FHeaderSection Header;
+		if (Header.BeginSection(Ar) && Header.SectionName == Header.kName && Header.EndSection(Ar))
+		{
+			return true;
+		}
+		return false;
+	}
+}
diff --git a/Engine/Source/Runtime/AudioExtensions/Private/IAudioCodecRegistry.cpp b/Engine/Source/Runtime/AudioExtensions/Private/IAudioCodecRegistry.cpp
new file mode 100644
--- /dev/null
+++ b/Engine/Source/Runtime/AudioExtensions/Private/IAudioCodecRegistry.cpp
@@ -0,0 +1,266 @@
+// Copyright Epic Games, Inc. All Rights Reserved.
+
+#include "IAudioCodecRegistry.h"
+#include "IAudioCodec.h"
+#include "Features/IModularFeatures.h"
+#include "Misc/CoreMisc.h"
+#include "Templates/SubclassOf.h"
+#include "DSP/DeinterleaveView.h"
+
+#include "Serialization/BulkData.h"
+#include "Serialization/BulkDataWriter.h"
+
+#include "GenericPlatform/GenericPlatformCriticalSection.h"
+
+DEFINE_LOG_CATEGORY(LogAudioCodec);
+
+namespace Audio
+{
+	class FCodecRegistry : public ICodecRegistry
+	{		
+		// Ownership.
+		TArray<TUniquePtr<ICodec>> LifetimeArray;				// Ownership is held here for all registered codecs
+
+		// Look ups. (these contain pointers to the above and must be updated when we add/remove from lifetime array).
+		TMap<FName, TArray<FCodecPtr>> NameLookupMap;			// UniqueName -> [ICodecs] (ordered by version) (i.e. PcmGeneric v1, PcmGeneric v2 etc).
+		TMap<FName, TArray<FCodecPtr>> FamilyNameLookupMap;		// FamilyName -> [ICodecs] (ordered by version) 
+
+		// RW lock on above containers.
+		mutable FRWLock ContainersLock;
+
+		//TMap<FName, TArray<FCodecPtr>> FamilyLookup;	// FamilyName -> [ICodecs]
+
+		// On Editor all other platform codecs with be discovered here
+		// Version|FamilyName|ImplName|		Flags| Features
+		// [0]    [Opus]     [OpusHwXbox]	[D]   [HwDecoder][HwResampler]
+		// [0]    [Opus]     [Opus]			[D+E]
+
+		// Anything matching the same version and family will be compatible.
+		// This should mean that we could test hardware/software implementations against each other.				
+
+		void InsertSortedByVersion(TArray<FCodecPtr>& InOutArray, FCodecPtr InCodecRawPtr)
+		{
+			// Store each identically named codec in order of version. (highest version first).
+			InOutArray.Emplace(InCodecRawPtr);
+			InOutArray.StableSort([](const ICodec& L, const ICodec& R) -> bool {
+				return L.GetDetails().Version > R.GetDetails().Version;
+			});
+		}
+
+		bool RegisterCodec(TUniquePtr<ICodec>&& InCodecPtr) override
+		{
+			if (!audio_ensure(InCodecPtr))
+			{
+				return false;
+			}
+
+			// Lock for writing.
+			FRWScopeLock ScopeLock(ContainersLock, SLT_Write);
+			
+			// Make sure we haven't got this pointer for some reason.
+			if( !audio_ensure(!LifetimeArray.Contains(InCodecPtr)) )
+			{
+				return false;
+			}
+
+			// Make sure we haven't already registered one with the same name and version.
+			if( !audio_ensure(!LifetimeArray.ContainsByPredicate([&InDetails = InCodecPtr->GetDetails()](const TUniquePtr<ICodec>& i)  { return i->GetDetails() == InDetails; }) ) )
+			{
+				return false;
+			}
+					   
+			// Take ownership
+			int32 Index = LifetimeArray.Emplace(MoveTemp(InCodecPtr));
+			if( audio_ensure(Index >= 0) && audio_ensure(LifetimeArray[Index] != nullptr) )
+			{	
+				FCodecPtr CodecRawPtr = LifetimeArray[Index].Get();
+				const FCodecDetails& Details = CodecRawPtr->GetDetails();
+					
+				// Insert in the look ups.
+				InsertSortedByVersion(NameLookupMap.FindOrAdd(Details.Name), CodecRawPtr);
+				InsertSortedByVersion(FamilyNameLookupMap.FindOrAdd(Details.FamilyName), CodecRawPtr);
+
+				UE_LOG(LogAudioCodec, Log, TEXT("Registered %s"), *Details.ToString());
+				return true;
+			}
+			return false;
+		}
+
+		bool UnregisterCodec(ICodecRegistry::FCodecPtr InCodecPtr) override
+		{
+			// Lock for writing.
+			FRWScopeLock ScopeLock(ContainersLock, SLT_Write);
+
+			if (!audio_ensure(InCodecPtr))
+			{
+				return false;
+			}
+
+			auto FindLambda = [InCodecPtr](const TUniquePtr<ICodec>& i) -> bool { return i.Get() == InCodecPtr; };
+
+			if (audio_ensure(LifetimeArray.ContainsByPredicate(FindLambda)))
+			{
+				UE_LOG(LogAudioCodec, Log, TEXT("Unregistered %s"), *InCodecPtr->GetDetails().ToString());
+
+				// Remove from lookups first. (maintain order).
+				TArray<FCodecPtr>& VersionArray = NameLookupMap.FindChecked(InCodecPtr->GetName());
+				VersionArray.Remove(InCodecPtr);
+				if( VersionArray.Num() == 0 )
+				{
+					NameLookupMap.Remove(InCodecPtr->GetName());
+				}
+
+				// Kill the object.
+				audio_ensure(LifetimeArray.RemoveAll(FindLambda) == 1);
+
+				return true;
+			}
+			return false;
+		}
+
+		FCodecPtr FindCodecByFamilyName(
+			FName InFamilyName, int32 InVersion) const override
+		{
+			// Lock for reading.
+			FRWScopeLock ScopeLock(ContainersLock, SLT_ReadOnly);
+			
+			// We might want to pass in the platform name of the codec, but for now assume we mean the current one.
+			FName InPlatformName = FPlatformMisc::GetUBTPlatform();
+
+			// Find matching family. 
+			if( const TArray<FCodecPtr> *pCodecArray = FamilyNameLookupMap.Find(InFamilyName) )
+			{
+				// Match version+platform.
+				FCodecPtr const * pCodec = pCodecArray->FindByPredicate([InVersion, InPlatformName](const FCodecPtr& i) -> bool 
+				{ 
+					// Version INDEX_NONE will return the latest version. 
+					return (i->GetDetails().Version == InVersion || InVersion == INDEX_NONE) && i->SupportsPlatform(InPlatformName);	
+				} );
+
+				if( pCodec )
+				{
+					return *pCodec;
+				}
+			}
+			return nullptr;
+		}
+
+		FCodecPtr FindCodecByName(
+			FName InName, int32 InVersion ) const override
+		{	
+			if(!audio_ensure(InName.IsValid()))
+			{
+				return nullptr;
+			}
+
+			// Lock for reading.
+			FRWScopeLock ScopeLock(ContainersLock, SLT_ReadOnly);
+
+			if (TArray<FCodecPtr> const *pVersionArray = NameLookupMap.Find(InName))
+			{
+				if (pVersionArray->Num() > 0)
+				{
+					// Asking for latest?
+					if (InVersion == INDEX_NONE)
+					{
+						// First item is the latest.
+						audio_ensure((*pVersionArray)[0]);
+						return (*pVersionArray)[0];
+					}
+					else
+					{
+						// Find a specific version. (todo: use binary search)
+						if (FCodecPtr const *pFound = pVersionArray->FindByPredicate([InVersion](const FCodecPtr& i) -> bool { return i->GetDetails().Version == InVersion; }))
+						{
+							audio_ensure(pFound);
+							return *pFound;
+						}
+					}
+				}				
+			}
+			return nullptr;
+		}
+		
+		virtual ICodecRegistry::FCodecPtr FindCodecByFromParsingInput(
+			IDecoderInput* InObject) const override
+		{
+			if( !audio_ensure(InObject))
+			{
+				return nullptr;
+			}
+			
+			// Look in the input stream for the header sections.
+			FFormatDescriptorSection Format;
+			if( !audio_ensure(InObject->FindSection(Format)) )
+			{
+				return nullptr;
+			}
+
+			FCodecPtr Ptr = FindCodecByFamilyName(
+				Format.CodecFamilyName,
+				Format.CodecVersion);
+
+			if( Ptr )
+			{
+				return Ptr;
+			}
+
+			// Fail.
+			return nullptr;		
+		}
+
+		FCodecPtr FindDefaultCodec(
+			FName InPlatformName) const override
+		{
+			// If we didn't specify a platform codec, use the current 
+			if( InPlatformName.IsNone() )
+			{
+				InPlatformName = FPlatformMisc::GetUBTPlatform();
+			}
+
+			// Temp. We need to have config settings to define this but for now, hard-code PCM.
+			static const FName NAME_Pcm = TEXT("Pcm");
+			if( FCodecPtr pPcmCodec = FindCodecByFamilyName(NAME_Pcm, INDEX_NONE))
+			{
+				return pPcmCodec;
+			}					   
+
+			UE_LOG(LogAudioCodec, Warning, TEXT("No default codec defined for %s"), *InPlatformName.GetPlainNameString());
+			return nullptr;
+		}
+
+	public:		
+		virtual ~FCodecRegistry()
+		{}
+
+		FCodecRegistry()
+		{
+			/*IModularFeatures::Get().OnModularFeatureRegistered().AddLambda([this](const FName& Type, IModularFeature* ModularFeature) 
+			{
+				if (Type == ICodec::GetModularFeatureName() && ModularFeature)
+				{
+					RegisterCodec(FCodecPtr(ModularFeature));
+				}
+			});
+
+			IModularFeatures::Get().OnModularFeatureUnregistered().AddLambda([this](const FName& Type, IModularFeature* ModularFeature) 
+			{
+				if (Type == ICodec::GetModularFeatureName() && ModularFeature)
+				{
+					UnregisterCodec(FCodecPtr(ModularFeature));
+				}
+			});			*/
+		}
+
+	};
+
+	ICodecRegistry& ICodecRegistry::Get()
+	{
+		// FIXME.
+		//static IModuleInterface* ExtentionsModule = FModuleManager::LoadModulePtr<IModuleInterface>(FName(TEXT("AudioExtensions")));
+
+		static FCodecRegistry Impl;
+		return Impl;
+	}
+
+} //namespace Audio
diff --git a/Engine/Source/Runtime/AudioExtensions/Private/IAudioDecoderScheduler.cpp b/Engine/Source/Runtime/AudioExtensions/Private/IAudioDecoderScheduler.cpp
new file mode 100644
--- /dev/null
+++ b/Engine/Source/Runtime/AudioExtensions/Private/IAudioDecoderScheduler.cpp
@@ -0,0 +1,174 @@
+// Copyright Epic Games, Inc. All Rights Reserved.
+
+#include "IAudioDecoderScheduler.h"
+#include "IAudioCodecRegistry.h"
+#include "IAudioCodec.h"
+#include "Templates/Atomic.h"
+
+namespace Audio
+{
+	struct FCmdQueue
+	{
+		using FCmdFunction = TFunction<void()>;
+		using FCommandQ = TArray<FCmdFunction>;
+
+		FCommandQ Q[2];
+		TAtomic<uint32> WriteIndex;
+		TAtomic<uint32> NumCmds;
+
+		void Execute()
+		{	
+			// Consumer thread.
+
+			// Swap the currently writing index.
+			uint32 CurrentIndex = WriteIndex;
+			while( !WriteIndex.CompareExchange(CurrentIndex,!WriteIndex) )
+			{
+				CurrentIndex = WriteIndex.Load(EMemoryOrder::SequentiallyConsistent);
+			}
+			
+			FCommandQ& CurrentQ = Q[CurrentIndex];
+			for (int32 i = 0; i < CurrentQ.Num(); ++i)
+			{
+				FCmdFunction Cmd(MoveTemp(CurrentQ[i]));
+				Cmd();
+				NumCmds--;
+			}
+			CurrentQ.Reset();
+		}
+		void QueueTask(
+			FCmdFunction&& InFunction)
+		{
+			// Producer thread.
+			Q[WriteIndex].Add(MoveTemp(InFunction));
+			NumCmds++;
+		}
+	};
+
+	struct FDecoderScheduler : public IDecoderScheduler
+	{
+		struct FActiveDecoder
+		{
+			enum EDecodeState
+			{
+				eStart,
+				eCreatingDecoder,
+				eParsingHeader,
+				eDecodingPackets,
+				eFinished
+			};
+
+			bool Update()
+			{
+				switch(State)
+				{
+					case eStart:
+					{
+						State = eCreatingDecoder;
+					} 
+					// fall-through.
+					case eCreatingDecoder:
+					{
+						ICodec* pCodec = ICodecRegistry::Get().FindCodecByFromParsingInput(
+							StartArgs.InputObject.Get());
+
+						if( !pCodec )
+						{
+							return false;
+						}
+						Decoder = pCodec->CreateDecoder(
+							StartArgs.InputObject.Get(),
+							StartArgs.OutputObject.Get() );							
+
+						State = eParsingHeader;
+					}
+					// fall-through.
+					case eParsingHeader:
+					{
+						/*FDecodedFormatInfo Info;												
+						if( !Decoder->ParseHeader(StartArgs.InputObject.Get(),Info))
+						{
+							return false;
+						}*/
+						State = eDecodingPackets;
+					}
+					// fall-through.
+					case eDecodingPackets:
+					{
+						// Do the decode.
+						IDecoder::EDecodeResult Result = Decoder->Decode();
+
+						if (Result == IDecoder::EDecodeResult::Finished || 
+							Result == IDecoder::EDecodeResult::Fail )
+						{
+							State = eFinished;
+						}
+						else
+						{
+							break;
+						}
+					}
+					// fall-through.
+					case eFinished:
+					{
+						break;
+					}
+					default: 
+					{
+						checkNoEntry();
+					}
+				}
+
+				// WIP. State machine logic.
+				return State != eFinished;
+			}
+
+			FStartDecodeArgs StartArgs;
+			TUniquePtr<IDecoder> Decoder;
+			EDecodeState State = eStart;
+		};
+
+		FCmdQueue CmdQueue;	
+		using FActiveMap = TMap<uint32, FActiveDecoder>;
+		TMap<uint32, FActiveDecoder> ActiveMap;
+
+		FDecoderScheduler() = default;
+		virtual ~FDecoderScheduler() = default;
+
+		void Update()
+		{
+			// Execute any cmds we have queued.
+			CmdQueue.Execute();
+
+			// Update active decodes.
+			for( auto i = ActiveMap.CreateIterator(); i; ++i )
+			{
+				if( !i->Value.Update() )
+				{
+					i.RemoveCurrent();
+				}
+			}
+		}
+
+		FActiveDecodeHandle StartDecode(
+			FStartDecodeArgs&& InArgs ) override
+		{
+			static TAtomic<uint32> sCounter;
+			uint32 Id = ++sCounter;
+
+			CmdQueue.QueueTask([this, Id, Args = MoveTemp(InArgs)] 
+			{
+				check(ActiveMap.Find(Id) == nullptr);
+				ActiveMap.Emplace(Id, FActiveDecoder { Args } );	
+			});
+			return FActiveDecodeHandle(Id);
+		}
+	};
+	
+	Audio::IDecoderScheduler& IDecoderScheduler::Get()
+	{
+		static FDecoderScheduler sInstance;
+		return sInstance;
+	}
+
+} //namespace Audio
\ No newline at end of file
diff --git a/Engine/Source/Runtime/AudioExtensions/Private/PcmCodec.cpp b/Engine/Source/Runtime/AudioExtensions/Private/PcmCodec.cpp
new file mode 100644
--- /dev/null
+++ b/Engine/Source/Runtime/AudioExtensions/Private/PcmCodec.cpp
@@ -0,0 +1,303 @@
+// Copyright Epic Games, Inc. All Rights Reserved.
+
+#include "PcmCodec.h"
+#include "IAudioCodec.h"
+#include "IAudioCodecRegistry.h"
+
+#include "Templates/SubclassOf.h"
+
+namespace Audio
+{	
+	class FPcmAudioCodec : public ICodec
+	{
+	public:
+		static const uint32	CodecVersion = 2;
+
+		struct FPcmHeader : public FEncodedSectionBase
+		{
+			static const uint32 kHeaderVer = 1;
+			static constexpr FFcc NAME_PcmHeader = { 'P','c','m','H' };
+
+			uint32 NumChannels = 0;
+			TEnumAsByte<Audio::EBitRepresentation>	BitFormat = EBitRepresentation::Int16_Interleaved;
+
+			FPcmHeader(uint32 InNumChannels = 0, EBitRepresentation InBitFormat = EBitRepresentation::Int16_Interleaved)
+				: FEncodedSectionBase(NAME_PcmHeader, kHeaderVer)
+				, NumChannels(InNumChannels)
+				, BitFormat(InBitFormat)
+			{}
+
+			void Serialize(FArchive& Ar) override
+			{
+				Ar << BitFormat;
+				Ar << NumChannels;
+			}
+		};
+
+		struct FEncoder : public IEncoder
+		{
+			IEncoder::FSourcePtr Input;
+
+			FEncoder(IEncoder::FSourcePtr InDataSource)
+				: Input(InDataSource)
+			{}
+
+			TSubclassOf<UAudioCodecEncoderSettings> GetSettingsClass() const override
+			{
+				return UAudioPcmEncoderSettings::StaticClass();
+			}
+
+			// Given input and settings determine DDC key.
+			FString GenerateDDCKey(
+				FSourcePtr SrcObject,
+				USettingsPtr SettingsObject) const override
+			{
+				check(SettingsObject);//.IsValid());
+
+				FString SrcHash = SrcObject->GetHashForDDC();
+				FString SettingsHash = SettingsObject->GetHashForDDC();			// Includes Version and Codec Name.
+				uint32 BinaryLibsHash = /* Some hash for the binaries */ 012345;
+
+				// Src + Settings + Codec.
+				return FString::Printf(TEXT("AudioCodec-%s-%s-%u"), *SrcHash, *SettingsHash, BinaryLibsHash);
+			}
+									
+			// Encode 
+			bool Encode(
+				FSourcePtr InSrc,
+				FDestPtr InDst,
+				USettingsPtr InSettingsObject ) const override
+			{
+				// Settings (use defaults if we were not fed anything).
+				const UAudioPcmEncoderSettings* PcmSettings = Cast<UAudioPcmEncoderSettings>(InSettingsObject);
+				if( !PcmSettings )
+				{					
+					PcmSettings = Cast<UAudioPcmEncoderSettings>(UAudioPcmEncoderSettings::StaticClass()->GetDefaultObject(true));
+				}
+
+				if (!InSrc || !InDst || !PcmSettings )
+				{
+					// TODO: Logging.
+					return false;
+				}
+
+				// Query input for details.
+				FDecodedFormatInfo Format;
+				if (!InSrc->GetFormat(Format))
+				{
+					// TODO: Logging.
+					return false;
+				}
+
+				const TSampleBuffer<int16>& SampleData = InSrc->GetSamples();
+
+				// Start header.
+				InDst->BeginHeader();
+				{
+					// Standard Descriptor Header, identifying codec and version of codec.
+					FFormatDescriptorSection Desc{ 
+						TEXT("PcmInternal"), 
+						TEXT("Pcm"), 
+						CodecVersion, 
+						Format.NumChannels, 
+						(uint32)SampleData.GetNumFrames(), 
+						Format.NumFramesPerSec, 
+						(uint32)sizeof(int16)*Format.NumChannels	// <-- FIXME. Each packet is just a frame, very small.
+					};
+					InDst->WriteSection(Desc);
+
+					// Format specific header.
+					FPcmHeader PcmHeader(Format.NumChannels, Format.BitRep);
+					InDst->WriteSection(PcmHeader);
+				}								
+				InDst->EndHeader();
+
+				//uint32 NumFramesPerPacket = 1;// (PcmSettings->PacketSizeBytes / sizeof(int16)) / SampleData.GetNumChannels();
+				//uint32 NumFramesRemaining = SampleData.GetNumFrames();
+				//uint32 NumFramesPerBlock = NumFramesPerPacket * PcmSettings->PacketsPerBlock;
+				
+				TArrayView<const int16> SampleAV = SampleData.GetArrayView();
+				uint32 SampleIndex = 0;
+
+				// Sample Data (hello world v1, write everything as one big blast).
+				InDst->BeginSampleData();			
+				InDst->WriteRaw(SampleAV.GetData(), SampleAV.Num()*SampleAV.GetTypeSize() );
+				InDst->EndSampleData();
+
+			/*	while (NumFramesRemaining > 0)
+				{
+					uint32 NumFramesToWrite = FMath::Min(NumFramesRemaining, NumFramesPerPacket);
+					TArrayView<const int16> PacketView = SampleAV.Slice(SampleIndex, NumFramesToWrite * SampleData.GetNumChannels());
+
+					InDst->WritePacket(PacketView.GetData(), PacketView.Num()*PacketView.GetTypeSize());
+
+					NumFramesRemaining -= NumFramesToWrite;
+					SampleIndex += NumFramesToWrite * SampleData.GetNumChannels();
+				}*/				
+
+				// Success.
+				return true;
+			}
+		};
+
+		template<typename TSampleType>
+		struct TDecoder : public IDecoder
+		{
+			bool bHasError = false;
+			FDecoderInputPtr Src;
+			FDecoderOutputPtr Dst;
+			FPcmHeader Header;
+			FFormatDescriptorSection Desc;
+			IDecoderOutput::FRequirements Reqs;
+			int32 PlayPositionOffsetFrames = 0;
+		
+			TDecoder( 
+				FDecoderInputPtr InSrc,
+				FDecoderOutputPtr InDst,
+				const FPcmHeader& InPcmHeader,
+				const FFormatDescriptorSection& InFormatDesc,
+				const IDecoderOutput::FRequirements& InRequirements )
+				: Src(InSrc)
+				, Dst(InDst)
+				, Header(InPcmHeader)
+				, Desc(InFormatDesc)
+				, Reqs(InRequirements)
+			{
+			}
+
+			bool HasError() const override
+			{
+				return bHasError;
+			}	
+
+			FDecodeReturn Decode() override
+			{
+				// Reject bad input
+				if( !audio_ensure(Src) || !audio_ensure(Dst) || !audio_ensure(Header.NumChannels > 0))
+				{
+					return FDecodeReturn { FDecodeReturn::Fail };
+				}
+
+				constexpr int32 NumFramesPerPacket = 1;
+				int32 PacketSizeBytes = (sizeof(TSampleType) * Header.NumChannels) * NumFramesPerPacket;
+				int32 FramesRemaining = Reqs.NumSampleFramesWanted;	
+				IDecoderOutput::FPushedAudioDetails Details(Desc.NumFramesPerSec,Desc.NumChannels, PlayPositionOffsetFrames);
+
+				while(!Src->IsEndOfStream() && FramesRemaining > 0)
+				{
+					TArrayView<const uint8> PacketBytes = Src->PopNextPacket(PacketSizeBytes);	
+					bool bLastPacket = PacketBytes.Num() == PacketSizeBytes;
+
+					TArrayView<const TSampleType> PacketPcm = MakeArrayView(
+						reinterpret_cast<const TSampleType*>(PacketBytes.GetData()), PacketBytes.Num() / sizeof(TSampleType));
+
+					Details.SampleFramesStartOffset = PlayPositionOffsetFrames;
+					Dst->PushAudio(Details, PacketPcm);
+					
+					FramesRemaining -= PacketPcm.Num();
+					PlayPositionOffsetFrames += PacketPcm.Num();
+				}
+				if( Src->IsEndOfStream() )
+				{
+					return FDecodeReturn::Finished;
+				}
+				return FDecodeReturn::MoreDataRemaining;
+			}
+		};
+
+		FPcmAudioCodec() = default;
+		virtual ~FPcmAudioCodec() = default;
+
+		bool SupportsPlatform(FName InPlatformName) const override
+		{
+			// All platforms support PCM.
+			return true;
+		}
+		const FCodecDetails& GetDetails() const override
+		{
+			static FCodecDetails sDeets = FCodecDetails({
+				TEXT("PcmInternal"),
+				TEXT("Pcm"),
+				CodecVersion,
+				FCodecFeatures({
+					FCodecFeatures::HasEncoder,
+					FCodecFeatures::HasDecoder,
+					FCodecFeatures::Seekable,
+					FCodecFeatures::Streamable,
+					FCodecFeatures::Surround
+				})
+			});
+			return sDeets;
+		}
+		FEncoderPtr CreateEncoder(
+			IEncoder::FSourcePtr InDataSource) override
+		{
+			return FEncoderPtr(new FEncoder(InDataSource));
+		}
+
+		virtual FDecoderPtr CreateDecoder(
+			IDecoder::FDecoderInputPtr InSrc,
+			IDecoder::FDecoderOutputPtr InDst )
+		{
+			FFormatDescriptorSection Desc;
+			if (!audio_ensure(InSrc->FindSection(Desc)))
+			{
+				return nullptr;
+			}
+			FPcmHeader PcmHeader;
+			if( !audio_ensure(InSrc->FindSection(PcmHeader)) )
+			{
+				return nullptr;
+			}
+			
+			FDecodedFormatInfo Info;
+			const IDecoderOutput::FRequirements& Reqs = InDst->GetRequirements(Info);
+
+			switch( PcmHeader.BitFormat )
+			{
+			case EBitRepresentation::Float32_Interleaved:
+			{
+				return FDecoderPtr( new TDecoder<float>( InSrc, InDst, PcmHeader, Desc, Reqs ) );
+			}
+			case EBitRepresentation::Int16_Interleaved:
+			{
+				return FDecoderPtr( new TDecoder<int16>( InSrc, InDst, PcmHeader, Desc, Reqs ) );
+			}
+			default:
+				checkNoEntry();
+			}
+
+			// Fail.
+			return nullptr;
+		}
+	};
+
+	TUniquePtr<Audio::ICodec> Create_PcmCodec()
+	{
+		return MakeUnique<FPcmAudioCodec>();
+	}
+} // namespace Audio
+
+//UAudioPcmEncoderSettings::UAudioPcmEncoderSettings(
+//	const FObjectInitializer& ObjectInitializer) 
+//	: Super(ObjectInitializer)
+//	, Version(1)
+//{}
+
+FString UAudioPcmEncoderSettings::GetHashForDDC() const
+{
+	//// Do this with reflection. (one probably exists already).
+	//// Also do base types in base class.
+	//uint32 Hash = HashCombine(
+	//	HashCombine(
+	//		GetTypeHash(PacketsPerBlock),
+	//		GetTypeHash(TestTweaker)
+	//	),
+	//	HashCombine(
+	//		GetTypeHash(MaxChannelsSupported),
+	//		GetTypeHash(PacketSizeBytes)
+	//	)
+	//);
+	int32 Hash = 123456789; // FIXME.
+	return FString::FromInt((int32)Hash);
+}
diff --git a/Engine/Source/Runtime/AudioExtensions/Private/PcmCodec.h b/Engine/Source/Runtime/AudioExtensions/Private/PcmCodec.h
new file mode 100644
--- /dev/null
+++ b/Engine/Source/Runtime/AudioExtensions/Private/PcmCodec.h
@@ -0,0 +1,32 @@
+// Copyright Epic Games, Inc. All Rights Reserved.
+
+#pragma once
+
+#include "IAudioCodec.h"
+#include "PcmCodec.generated.h"
+
+UENUM()
+enum class EPcmBitDepthConversion : uint8
+{
+	SameAsSource,
+	Int16,
+	Float32
+};
+
+UCLASS()
+class AUDIOEXTENSIONS_API UAudioPcmEncoderSettings : public UAudioCodecEncoderSettings
+{
+public:
+	GENERATED_BODY()	
+
+	UPROPERTY()
+	EPcmBitDepthConversion BitDepthConversion;
+
+protected:
+	virtual FString GetHashForDDC() const override;
+};
+
+namespace Audio
+{
+	TUniquePtr<ICodec> Create_PcmCodec();
+}
\ No newline at end of file
diff --git a/Engine/Source/Runtime/AudioExtensions/Public/IAudioCodec.h b/Engine/Source/Runtime/AudioExtensions/Public/IAudioCodec.h
new file mode 100644
--- /dev/null
+++ b/Engine/Source/Runtime/AudioExtensions/Public/IAudioCodec.h
@@ -0,0 +1,728 @@
+// Copyright Epic Games, Inc. All Rights Reserved.
+
+#pragma once
+
+#include "CoreMinimal.h"
+#include "Features/IModularFeatures.h"
+#include "Modules/ModuleInterface.h"
+#include "Modules/ModuleManager.h"
+
+#include "DSP/BufferVectorOperations.h"
+#include "DSP/Dsp.h"
+#include "SampleBuffer.h"
+#include "Serialization/MemoryReader.h"
+
+#include "IAudioCodec.generated.h"
+
+AUDIOEXTENSIONS_API DECLARE_LOG_CATEGORY_EXTERN(LogAudioCodec, Display, All);
+
+// TODO: Move this somewhere better
+#ifdef AUDIO_ENSURE_NO_BREAKS
+	#define audio_ensure(EXP) (!!(EXP))
+#else //AUDIO_ENSURE_NO_BREAKS
+	#define audio_ensure(EXP) ensure(EXP)
+#endif //AUDIO_VERIFY_NO_BREAKS
+
+// Forward declares.
+class USoundWave;
+
+#pragma region Decoder
+
+namespace Audio
+{
+	// Forward declares
+	struct FEncodedSectionBase;
+	struct FMixerSourceVoiceBuffer;
+
+	// Four character code class.
+	class FFcc
+	{
+	public:
+		FFcc() = default;
+
+#if 0 // Clang doesn't like this.
+		template<uint32 N>
+		constexpr FFcc(const char(&InStr)[N])
+			: Value(0)
+		{
+			static_assert(N <= 5, "Four Or Less Characters Strings Only");
+			for (uint32 i = 0, Shift = 0; i < N; ++i, Shift += 8)
+			{
+				Value |= static_cast<uint32>(InStr[i] << Shift);
+			}
+		}
+#endif 
+		constexpr FFcc(const char One, const char Two, const char Three, const char Four)
+		{
+			Value = static_cast<uint32>(One) |
+				(static_cast<uint32>(Two) << 8) |
+				(static_cast<uint32>(Three) << 16) |
+				(static_cast<uint32>(Four) << 24);
+		}
+
+		constexpr bool operator==(const FFcc& InRhs) const
+		{
+			return Value == InRhs.Value;
+		}
+		friend const FFcc& operator<<(FArchive& Ar, FFcc& InFcc)
+		{
+			Ar << InFcc.Value;
+			return InFcc;
+		}
+		constexpr bool IsNone() const { return Value==0; }
+		friend uint32 GetTypeHash(FFcc InVal)
+		{
+			return ::GetTypeHash(InVal.Value);
+		}
+	private:
+		uint32 Value = 0;
+	};
+	
+	enum EBitRepresentation
+	{
+		Int16_Interleaved,
+		Float32_Interleaved
+	};
+
+	struct AUDIOEXTENSIONS_API FDecodedFormatInfo 
+	{		
+		uint32 NumChannels = 0;
+		uint32 NumFramesPerSec = 0;
+		EBitRepresentation BitRep = Int16_Interleaved;
+
+		int32 LoopStartOffset = INDEX_NONE;	// Where the loop start point is in the entire stream
+		int32 LoopEndOffset = INDEX_NONE;	// Where the loop end point is in the entire stream
+	};
+	
+	struct AUDIOEXTENSIONS_API IDecoderInput
+	{
+		// Factory.
+		static TUniquePtr<IDecoderInput> Create(
+			TArrayView<const uint8> InBytes, uint32 InOffsetOfBytes, FName InOldFormatName=NAME_None );
+						
+		virtual ~IDecoderInput() {}
+
+		virtual bool HasError() const = 0;
+		virtual bool IsEndOfStream() const = 0;
+
+		virtual bool FindSection(
+			FEncodedSectionBase& OutSection) = 0;
+
+		virtual int64 Tell() const = 0;
+
+		virtual TArrayView<const uint8> PeekNextPacket(
+			int32 InMaxPacketLength) const = 0;
+				
+		virtual TArrayView<const uint8> PopNextPacket(
+			int32 InPacketSize ) = 0;
+	};
+	
+	struct AUDIOEXTENSIONS_API IDecoderOutput
+	{	
+		struct FRequirements
+		{
+			EBitRepresentation DownstreamFormat	= Float32_Interleaved;	// Our pipeline is all float downstream, so this might not be necessary.
+			int32 NumSampleFramesWanted			= 256;
+		};
+
+		// Factory.
+		static TUniquePtr<IDecoderOutput> Create( 
+			const FRequirements& Requirements );
+
+		// Factory.
+		static TUniquePtr<IDecoderOutput> Create(
+			const FRequirements& Requirements, TArrayView<uint8> InExternalMemory );
+		
+		struct FPushedAudioDetails
+		{
+			int32 NumFramesPerSec;					
+			int32 NumChannels;
+
+			int32 SampleFramesStartOffset;				// Where these samples are from in the entire stream
+
+			uint32 bContainsLoopPoint : 1;				// True if *this* set of samples contains the loop point
+			uint32 bHasFinished : 1;					// True if this the last of the data
+			uint32 bHasError : 1;						// True if the decoder has encountered an error
+
+			FPushedAudioDetails(
+				int32 InFramesPerSec = 0,
+				int32 InNumChannels = 0,
+				int32 InSampleOffset = 0,
+				bool In_bContainsLoopPoint = false,
+				bool In_bHasFinished = false,
+				bool In_bHasError = false )			
+				: NumFramesPerSec(InFramesPerSec) 
+				, NumChannels(InNumChannels)
+				, SampleFramesStartOffset(InSampleOffset)
+				, bContainsLoopPoint(In_bContainsLoopPoint)
+				, bHasFinished(In_bHasFinished)
+				, bHasError(In_bHasError)
+			{}
+		};
+
+		virtual ~IDecoderOutput() = default;
+		
+		virtual FRequirements GetRequirements(
+			const FDecodedFormatInfo& InFormat) const = 0;
+		
+		virtual int32 PushAudio(
+			const FPushedAudioDetails& InDetails, 
+			TArrayView<const int16> In16BitInterleave ) = 0;
+
+		virtual int32 PushAudio(
+			const FPushedAudioDetails& InDetails,
+			TArrayView<const float> InFloat32Interleave) = 0;
+
+		virtual int32 PopAudio( 
+			TArrayView<int16>& OutAudio ) = 0;
+
+		virtual int32 PopAudio(
+			TArrayView<float>& OutAudio ) = 0;
+
+		virtual int32 Tell() const = 0;
+	};
+
+	// Decode input wrappers.
+
+	struct AUDIOEXTENSIONS_API FCodecFeatures
+	{	
+		bool HasFeature(FName Feature) const { return true; }
+		void GetFeatures(TArray<FName>& OutFeatureNames) const {}
+
+		enum EFeatureBitField
+		{
+			HwDecode,
+			HwResample,
+		
+			Seekable,
+			Streamable,
+			Surround,
+			SeamlessLooping,
+		
+			HasDecoder,
+			HasEncoder,
+		};
+		uint32_t FeaturesBitField;		// Store as bitfield.
+
+		bool HasFeature(EFeatureBitField InBit) const { return ( FeaturesBitField & (uint32)InBit); };
+
+		FCodecFeatures(std::initializer_list<EFeatureBitField> InitList)
+		{
+			FeaturesBitField = 0;
+			for (EFeatureBitField i : InitList)
+			{
+				FeaturesBitField |= 1 << i;
+			}
+		}
+	};
+
+	struct AUDIOEXTENSIONS_API IDecoder
+	{		
+		enum EDecodeResult
+		{
+			Fail,					// Decoder failed somehow.
+			MoreDataRemaining,		// Data has produced and there's more remaining
+			Looped,					// Data has produced and we've also looped
+			Busy,					// The decoder may not have produced anything, and is currently busy
+			Finished				// The decoder has finished and can be deleted
+		};
+
+		using FDecoderInputPtr = IDecoderInput*;
+		using FDecoderOutputPtr = IDecoderOutput*;
+		using FDecodeReturn = EDecodeResult;
+					
+		virtual ~IDecoder() = default;
+		virtual bool HasError() const = 0;
+
+		// Do the decode.
+		virtual FDecodeReturn Decode() = 0;	
+	};
+
+	struct AUDIOEXTENSIONS_API FEncodedSectionBase
+	{
+		// Serialized members.
+		FFcc SectionName;
+		uint32 SectionSize = 0;
+		uint32 SectionVersion = 0;
+
+		// This is not serialized and is used for accounting during serialization
+		uint32 SizePos = 0;
+
+		FEncodedSectionBase(FFcc InSectionName,uint32 InVersion) 
+			: SectionName(InSectionName)
+			, SectionVersion(InVersion)
+		{}
+
+		virtual ~FEncodedSectionBase() {}
+
+		virtual bool BeginSection(FArchive& Ar);
+		virtual bool EndSection(FArchive& Ar);
+		virtual void Serialize(FArchive& Ar) = 0;
+	};		
+
+	struct AUDIOEXTENSIONS_API FFormatDescriptorSection : public FEncodedSectionBase
+	{
+		static constexpr FFcc kSectionName = { 'D','e','s','c' };
+		static const uint32 kSectionVer = 2;
+		
+		// Codec name, family and version.
+		FName CodecName;
+		FName CodecFamilyName;
+		uint32 CodecVersion = 0;
+
+		// Info on the compressed data
+		uint32 NumChannels			= 0;
+		uint32 NumFrames			= 0;
+		uint32 NumFramesPerSec		= 0;
+		uint32 NumBytesPerPacket	= 0;
+				
+		FFormatDescriptorSection() 
+			: FEncodedSectionBase(kSectionName,kSectionVer)
+		{}
+		
+		FFormatDescriptorSection(
+			FName InCodecName,
+			FName InCodecFamilyName,
+			uint32 InCodecVersion,
+			uint32 InNumChannels,
+			uint32 InNumSamples,
+			uint32 InNumSamplesPerSec,
+			uint32 InNumBytesPerPacket
+		);
+
+		bool operator==(const FFormatDescriptorSection& Rhs) const 
+		{			
+			return CodecName		==	Rhs.CodecName
+			&& CodecFamilyName		==	Rhs.CodecFamilyName
+			&& CodecVersion			==	Rhs.CodecVersion
+			&& NumChannels			==	Rhs.NumChannels
+			&& NumFrames			==	Rhs.NumFrames
+			&& NumFramesPerSec		==	Rhs.NumFramesPerSec;		
+		}
+		
+		void Serialize(FArchive& Ar) override;
+	};
+
+	struct AUDIOEXTENSIONS_API FNestedSection : public FEncodedSectionBase
+	{
+		FNestedSection(FFcc InName = FFcc())  : FEncodedSectionBase(InName, 0) {}
+		void Serialize(FArchive& Ar) override {}
+	};
+
+	struct AUDIOEXTENSIONS_API FHeaderSection : public FNestedSection
+	{
+		static constexpr FFcc kName = {'H','d','r','!' };
+		FHeaderSection() : FNestedSection(kName) {}
+		void Serialize(FArchive& Ar) override {}
+	};
+	struct AUDIOEXTENSIONS_API FSampleSection : public FNestedSection
+	{
+		static constexpr FFcc kName = {'S','m','p','l'};
+		FSampleSection() : FNestedSection(kName) {}
+		void Serialize(FArchive& Ar) override {}
+	};
+	
+	struct AUDIOEXTENSIONS_API FRawSection : public FEncodedSectionBase
+	{
+		TArrayView<const uint8> Buffer;
+		FRawSection(FFcc InName, TArrayView<const uint8> InBuffer = TArrayView<const uint8>() )
+			: FEncodedSectionBase(InName, 0) 
+			, Buffer(InBuffer)
+		{}
+
+		void Serialize(FArchive& Ar) override {}
+	};
+
+	class AUDIOEXTENSIONS_API FDecoderInputBase : public Audio::IDecoderInput
+	{
+	protected:
+		FArchive* Archive = nullptr;
+		TArray<uint8> PopBuffer;
+		TMap<FFcc, int32> Toc;
+		uint32 OffsetInStream = 0;
+
+		bool bSampleSection = false;
+		bool bError = false;
+
+		bool SeekToSamplesStart(
+			FArchive& Ar );
+		
+		bool ParseSections(
+			FArchive& Ar);
+
+		bool MakeTocNested(
+			FArchive &Ar, uint32 HeaderStart, FNestedSection &Header);
+
+	public: 
+		void SetArchive(
+			FArchive* InArchive, uint32 InOffsetInStream);
+
+		bool FindSection(
+			FEncodedSectionBase& OutSection) override;
+
+		bool HasError() const override { return bError && Archive->IsError(); }
+		bool IsEndOfStream() const override { return Archive->AtEnd(); }
+
+		bool BeginSampleData();		
+		bool EndSampleData();
+
+		TArrayView<const uint8> PeekNextPacket(
+			int32 InMaxPacketLength) const override;
+
+		TArrayView<const uint8> PopNextPacket(
+			int32 InPacketSize) override;	
+
+		// We need this for back compatibility.
+		int64 Tell() const override;
+	};
+
+	struct AUDIOEXTENSIONS_API FDecoderInputArrayView : public FDecoderInputBase
+	{
+		TUniquePtr<FArchive> Reader;
+
+		static bool IsValidHeader(
+			TArrayView<const uint8> InBytes);
+
+		void SetArrayViewBytes(TArrayView<const uint8> InBytes, uint32 InOffsetInStream)
+		{
+			Reader = MakeUnique<FMemoryReaderView>(InBytes);
+			SetArchive(Reader.Get(), InOffsetInStream);
+		}
+		
+		// Pass in array view and we'll wrap around it
+		FDecoderInputArrayView( 
+			TArrayView<const uint8> InData, uint32 InOffsetInStream )
+			: Reader(MakeUnique<FMemoryReaderView>(InData))
+		{		
+			SetArchive(Reader.Get(), InOffsetInStream);
+		}
+
+		// Pass in an archive
+		FDecoderInputArrayView( 
+			TUniquePtr<FArchive>&& InReader, uint32 InOffsetInStream )
+			: Reader(MoveTemp(InReader))
+		{
+			SetArchive(Reader.Get(), InOffsetInStream);
+		}
+	};
+		
+	template<typename TSampleType>
+	struct TDecoderOutputArrayView : public IDecoderOutput
+	{
+		FRequirements Reqs;
+		TArrayView<TSampleType> Buffer;
+		int64 Offset = 0;
+
+		TDecoderOutputArrayView(const FRequirements& InReqs)
+			: Reqs(InReqs)
+			, Buffer()
+		{}
+
+		void SetArrayView(TArrayView<TSampleType> InBuffer)
+		{
+			Buffer = InBuffer;
+		}
+		void SetArrayViewBytes(TArrayView<uint8> InBytes)
+		{
+			SetArrayView(MakeArrayView(reinterpret_cast<TSampleType*>(InBytes.GetData()), InBytes.Num() / sizeof(TSampleType)));
+		}
+
+		FRequirements GetRequirements(
+			const FDecodedFormatInfo& InFormat) const override
+		{
+			return Reqs;
+		}
+
+		template<typename From, typename To>
+		static constexpr float GetConvertScalar()
+		{
+			if( TIsSame<From,To>::Value ) 
+			{
+				return 1.0f;
+			}
+			else if( TIsSame<From,int16>::Value && TIsSame<To,float>::Value)
+			{
+				return 1.f / static_cast<float>(FMath::Abs(TNumericLimits<int16>::Min()));
+			}
+			else if( TIsSame<From,float>::Value && TIsSame<To,int16>::Value)
+			{
+				return static_cast<float>(FMath::Abs(TNumericLimits<int16>::Min()));
+			}
+		}
+
+		template<typename T> int32 PushAudioInternal(
+			const FPushedAudioDetails& InDetails,
+			TArrayView<const T> InAudio) 
+		{
+			const T* Src = InAudio.GetData();
+			TSampleType* Dst = Buffer.GetData() + Offset;
+
+			int32 AvailSpace = Buffer.Num() - Offset;
+			int32 Num = FMath::Min(AvailSpace, InAudio.Num());
+
+			if (TIsSame<TSampleType, T>::Value)
+			{
+				FMemory::Memcpy(Dst, Src, Num * sizeof(T));
+			}
+			else
+			{
+				constexpr float Scalar = GetConvertScalar<T,TSampleType>();
+				for (int32 i = 0; i < Num; ++i)
+				{
+					*Dst++ = static_cast<TSampleType>(*Src++ * Scalar);
+				}
+			}
+
+			Offset += Num;
+			return Num;
+		}
+
+		int32 PushAudio(
+			const FPushedAudioDetails& InDetails,
+			TArrayView<const float> InFloat32Interleaved ) override
+		{
+			return PushAudioInternal(InDetails,InFloat32Interleaved );
+		}
+		
+		int32 PushAudio(
+			const FPushedAudioDetails& InDetails,
+			TArrayView<const int16> In16BitInterleave) override
+		{
+			return PushAudioInternal(InDetails,In16BitInterleave);
+		}
+	
+		template<typename T> int32 PopAudioInternal(
+			TArrayView<T>& OutSamples )
+		{
+			if (TIsSame<T, TSampleType>::Value && Offset > 0)
+			{
+				int32 NumSamplesPopped = Offset;
+				Offset = 0;
+				TArrayView<TSampleType> AV = Buffer.Slice(0, NumSamplesPopped);
+				OutSamples = MakeArrayView((T*)AV.GetData(), AV.Num());
+				return NumSamplesPopped;
+			}
+			
+			audio_ensure(false);
+			OutSamples = TArrayView<T>();
+			
+			return 0;
+		}
+
+		virtual int32 PopAudio(
+			TArrayView<float>& OutFloat )
+		{
+			return PopAudioInternal(OutFloat);
+		}
+
+		virtual int32 PopAudio(			
+			TArrayView<int16>& OutInt16 )
+		{
+			return PopAudioInternal(OutInt16);
+		}
+
+		// For back compatibility only.
+		int32 Tell() const 
+		{
+			return Offset;
+		}
+	};
+
+	template<typename T>
+	struct TDecoderOutputOwnBuffer : public TDecoderOutputArrayView<T>
+	{
+		TArray<T, TAlignedHeapAllocator<SIMD_ALIGNMENT>> AlignedMemory;
+		TDecoderOutputOwnBuffer(
+			const IDecoderOutput::FRequirements& Requirements)
+			: TDecoderOutputArrayView<T>(Requirements)
+		{
+			AlignedMemory.SetNumZeroed(Requirements.NumSampleFramesWanted);
+			TDecoderOutputArrayView<T>::SetArrayView(MakeArrayView(AlignedMemory));
+		}
+	};
+};
+
+#pragma endregion Decoder
+
+#pragma region Encoder
+
+UCLASS(Abstract)
+class AUDIOEXTENSIONS_API UAudioCodecEncoderSettings : public UObject
+{
+public:
+	GENERATED_BODY()
+
+	UPROPERTY()
+	int32 Version = 0;
+
+	// TODO: Ideally we could use reflection generate a hash of all the settings.
+	// ... But we could just override this per type.
+	virtual FString GetHashForDDC() const PURE_VIRTUAL(UAudioCodecEncoderSettings::GetHashForDDC, return TEXT(""););
+};	
+
+namespace Audio
+{	
+	class AUDIOEXTENSIONS_API IEncoderOutput
+	{
+	public:
+		using FPtr = TUniquePtr<IEncoderOutput>;
+		virtual ~IEncoderOutput() = default;
+
+		// Create one.
+		static FPtr Create();
+		static FPtr Create(TArray<uint8>& OutBuffer);
+
+		// Named section "format" for example.
+		virtual bool WriteSection(FEncodedSectionBase& InSection) = 0;
+		
+		virtual bool WriteRaw(const void* InRawData, int32 InSize) = 0;
+
+		// Header. (header is a bunch of named sections).
+		virtual bool BeginHeader() = 0;
+		virtual bool EndHeader() = 0;
+
+		// Sample data. (a series of named/chunks).
+		virtual bool BeginSampleData() = 0;
+		virtual bool EndSampleData() = 0;
+	};
+
+	class AUDIOEXTENSIONS_API IEncoderInput
+	{
+	public:
+		using FFormat = FDecodedFormatInfo;
+
+		// Factory.
+		static TUniquePtr<IEncoderInput> Create(
+			TArrayView<const int16> InSamples,
+			const FFormat& InInfo
+		);
+
+		virtual ~IEncoderInput() {}
+
+		// Calculate hash for input source.
+		virtual FString GetHashForDDC() const = 0;
+
+		// Get Format.
+		virtual bool GetFormat(
+			FFormat& OutFormat) const = 0;
+
+		virtual const TSampleBuffer<int16>& 
+			GetSamples() const = 0;
+	};
+
+	struct AUDIOEXTENSIONS_API IEncoder
+	{
+		virtual ~IEncoder() {}
+
+		// Get the class used for this encoder. 
+		virtual TSubclassOf<UAudioCodecEncoderSettings> GetSettingsClass() const = 0;
+
+		// Types.
+		using FSourcePtr = IEncoderInput * ;
+		using FDestPtr = IEncoderOutput * ;
+		using USettingsPtr = const UAudioCodecEncoderSettings *;
+
+		// look at the ISoundfieldFormat IEncodingSoundfieldSetttingProxy! :)
+
+		// Given input and settings determine DDC key.
+		virtual FString GenerateDDCKey(
+			FSourcePtr Src,
+			USettingsPtr SettingsObject) const = 0;
+
+		// Encode 
+		virtual bool Encode(
+			FSourcePtr InSrc,
+			FDestPtr InDst,
+			USettingsPtr InSettingsObject) const = 0;
+	};
+}
+
+#pragma endregion Encoder
+
+#pragma region Codec
+
+namespace Audio
+{
+	struct AUDIOEXTENSIONS_API FCodecDetails
+	{
+		FName Name;								// Unique name for this implementation. e.g. "OpusXbox"
+		FName FamilyName;						// Family name e.g. "Opus". Same as old Format name.	
+		uint32 Version = 0;						// For Compatibility (all codecs in this family will be compatible with this data).	
+		FCodecFeatures Features;				// If defined, is the feature set for this codec 	
+
+		FString ToString() const { return FString::Printf(TEXT("%s:%s:%d")); }
+		
+		bool operator==(const FCodecDetails& Rhs) const 
+		{ 
+			return Rhs.FamilyName == FamilyName &&
+				   Rhs.Name == Name &&
+				   Rhs.Version == Version;
+		}
+
+		FCodecDetails(FName InName, FName InFamilyName, uint32 InVersion, const FCodecFeatures& InFeatures )
+			: Name(InName), FamilyName(InFamilyName), Version(InVersion), Features(InFeatures) 
+		{}
+	};
+
+	struct AUDIOEXTENSIONS_API ICodec : public IModularFeature
+	{
+		virtual ~ICodec() {}
+
+		static FName GetModularFeatureName() 
+		{ 
+			static const FName sName = TEXT("AudioCodec"); 
+			return sName; 
+		}	
+
+		// TODO. Add lifetime subscription delegates.
+		// OnDecoderCreated 
+		// OnEncoderCreated
+		// OnStaticInit
+		// OnStaticDeint
+		
+		// Query.
+		virtual bool SupportsPlatform(FName InPlatformName) const = 0;		// This could move to the details.
+		virtual const FCodecDetails& GetDetails() const = 0;				// Name/Version/ImplName/Features.
+		FName GetName() const { return GetDetails().Name; };
+
+		// Factory for encoders
+		using FEncoderPtr = TUniquePtr<IEncoder>;
+		virtual FEncoderPtr CreateEncoder(
+			IEncoder::FSourcePtr) 
+		{ 
+			return nullptr; 
+		}
+	
+		// Factory for decoders
+		using FDecoderPtr = TUniquePtr<IDecoder>;
+		virtual FDecoderPtr CreateDecoder(
+			IDecoder::FDecoderInputPtr InSrc,
+			IDecoder::FDecoderOutputPtr InDst)
+		{ 
+			return nullptr; 
+		}			
+
+	protected:
+		// Register Modular features.
+		void RegisterModularFeature() { IModularFeatures::Get().RegisterModularFeature(ICodec::GetModularFeatureName(), this); }
+		void UnRegisterModularFeature() { IModularFeatures::Get().UnregisterModularFeature(ICodec::GetModularFeatureName(), this); }
+	};
+
+	// template this with the codec class.
+	// Convenience class for when you have just one codec in a module.
+	struct AUDIOEXTENSIONS_API IAudioCodecModule : public ICodec,
+												   public IModuleInterface
+	{
+		virtual void StartupModule() override
+		{
+			RegisterModularFeature();
+		}
+		virtual void ShutdownModule() override
+		{
+			UnRegisterModularFeature();
+		}
+	};
+
+} // namespace Audio
+
+#pragma endregion Codec
\ No newline at end of file
diff --git a/Engine/Source/Runtime/AudioExtensions/Public/IAudioCodecRegistry.h b/Engine/Source/Runtime/AudioExtensions/Public/IAudioCodecRegistry.h
new file mode 100644
--- /dev/null
+++ b/Engine/Source/Runtime/AudioExtensions/Public/IAudioCodecRegistry.h
@@ -0,0 +1,36 @@
+// Copyright Epic Games, Inc. All Rights Reserved.
+#pragma once
+
+#include "Templates/SharedPointer.h"
+
+#include "IAudioCodec.h"
+
+class ICompressedAudioInfo;
+
+namespace Audio
+{
+	class AUDIOEXTENSIONS_API ICodecRegistry
+	{
+	public:
+		using FCodecPtr = ICodec * ;
+	
+		// Abstract Singleton.
+		static ICodecRegistry& Get();
+
+		// Register each codec.
+		virtual bool RegisterCodec(TUniquePtr<ICodec>&&) = 0;
+		virtual bool UnregisterCodec(FCodecPtr) = 0;
+
+		// Find exact code name. (Version can be INDEX_NONE and give the latest version).
+		virtual FCodecPtr FindCodecByName(FName InName, int32 InVersion = INDEX_NONE) const = 0;
+
+		// Find a codec given a decoders input interface
+		virtual FCodecPtr FindCodecByFromParsingInput(IDecoderInput* InObject) const = 0;
+
+		// Find codec by finding by family name.  (Version can be INDEX_NONE and give the latest version).
+		virtual FCodecPtr FindCodecByFamilyName(FName InFamilyName, int32 InVersion = INDEX_NONE) const = 0;
+
+		// Find default codec for a platform. (none uses host platform).
+		virtual FCodecPtr FindDefaultCodec(FName InPlatformName = NAME_None) const = 0;
+	};
+}
diff --git a/Engine/Source/Runtime/AudioExtensions/Public/IAudioDecoderScheduler.h b/Engine/Source/Runtime/AudioExtensions/Public/IAudioDecoderScheduler.h
new file mode 100644
--- /dev/null
+++ b/Engine/Source/Runtime/AudioExtensions/Public/IAudioDecoderScheduler.h
@@ -0,0 +1,36 @@
+// Copyright Epic Games, Inc. All Rights Reserved.
+#pragma once
+
+#include "HAL/Platform.h"
+#include "Templates/SharedPointer.h"
+
+namespace Audio
+{
+	struct IDecoderInput;
+	struct IDecoderOutput;
+		
+	class FActiveDecodeHandle
+	{
+		uint32 Id;
+	public:
+		FActiveDecodeHandle(uint32 InId) 
+			: Id(InId) 
+		{}
+		bool IsValid() const { return Id != 0; }
+	};
+
+	struct IDecoderScheduler
+	{
+		static IDecoderScheduler& Get();
+
+		struct FStartDecodeArgs
+		{
+			uint32 StartingFrame = 0;
+			TSharedPtr<IDecoderInput> InputObject;
+			TSharedPtr<IDecoderOutput> OutputObject;
+		};
+
+		virtual FActiveDecodeHandle StartDecode(
+			FStartDecodeArgs&& ) = 0;
+	};
+}
\ No newline at end of file
