Fix for static analysis warning around Octal digits.

diff --git a/Engine/Source/Runtime/AudioExtensions/Private/PcmCodec.cpp b/Engine/Source/Runtime/AudioExtensions/Private/PcmCodec.cpp
--- a/Engine/Source/Runtime/AudioExtensions/Private/PcmCodec.cpp
+++ b/Engine/Source/Runtime/AudioExtensions/Private/PcmCodec.cpp
@@ -56,7 +56,7 @@ namespace Audio
 
 				FString SrcHash = SrcObject->GetHashForDDC();
 				FString SettingsHash = SettingsObject->GetHashForDDC();			// Includes Version and Codec Name.
-				uint32 BinaryLibsHash = /* Some hash for the binaries */ 012345;
+				uint32 BinaryLibsHash = /* Some hash for the binaries */ 12345;
 
 				// Src + Settings + Codec.
 				return FString::Printf(TEXT("AudioCodec-%s-%s-%u"), *SrcHash, *SettingsHash, BinaryLibsHash);
