Fix for CIS Linux/Clang linker error with FFcc constexpr

diff --git a/Engine/Source/Runtime/AudioExtensions/Private/PcmCodec.cpp b/Engine/Source/Runtime/AudioExtensions/Private/PcmCodec.cpp
--- a/Engine/Source/Runtime/AudioExtensions/Private/PcmCodec.cpp
+++ b/Engine/Source/Runtime/AudioExtensions/Private/PcmCodec.cpp
@@ -276,6 +276,10 @@ namespace Audio
 	{
 		return MakeUnique<FPcmAudioCodec>();
 	}
+
+	// Clang < C++ 17 requires constexp require storage for linkage
+	constexpr FFcc FPcmAudioCodec::FPcmHeader::NAME_PcmHeader;
+
 } // namespace Audio
 
 //UAudioPcmEncoderSettings::UAudioPcmEncoderSettings(
