Fixes for destructive unit tests and API refactor around PopBuffer taking external memory.

diff --git a/Engine/Plugins/Runtime/Metasound/Source/MetasoundStandardNodes/Private/MetasoundWavePlayerNode.cpp b/Engine/Plugins/Runtime/Metasound/Source/MetasoundStandardNodes/Private/MetasoundWavePlayerNode.cpp
--- a/Engine/Plugins/Runtime/Metasound/Source/MetasoundStandardNodes/Private/MetasoundWavePlayerNode.cpp
+++ b/Engine/Plugins/Runtime/Metasound/Source/MetasoundStandardNodes/Private/MetasoundWavePlayerNode.cpp
@@ -82,7 +82,7 @@ namespace Metasound
 		
 		void Execute()
 		{
-			int32 i = 0;
+			int32 NumPopped = 0;
 			float* Dst = AudioBuffer->GetData();
 
 			// If we don't have a valid state, just output silence.
@@ -91,16 +91,12 @@ namespace Metasound
 				// V1. Do the decode inline, this will sound bad.
 				Decoder->Decode();
 
-				TArrayView<float> Src;
-				int32 NumPopped = DecoderOutput->PopAudio(Src);
-				int32 NumFramesToCopy = FMath::Min(NumPopped, OperatorSettings.GetNumFramesPerBlock());
-
-				for (; i < NumFramesToCopy; ++i)
-				{
-					Dst[i] = Src[i];
-				}
+				Audio::IDecoderOutput::FPushedAudioDetails Details;
+				NumPopped = DecoderOutput->PopAudio(MakeArrayView(Dst, AudioBuffer->Num()), Details);
 			}
-			for ( ; i < OperatorSettings.GetNumFramesPerBlock(); ++i)
+
+			// Pad with Silence if we didn't pop enough
+			for ( int32 i = NumPopped; i < OperatorSettings.GetNumFramesPerBlock(); ++i)
 			{
 				Dst[i] = 0.0f;
 			}
diff --git a/Engine/Source/Runtime/AudioExtensions/Private/IAudioCodec.cpp b/Engine/Source/Runtime/AudioExtensions/Private/IAudioCodec.cpp
--- a/Engine/Source/Runtime/AudioExtensions/Private/IAudioCodec.cpp
+++ b/Engine/Source/Runtime/AudioExtensions/Private/IAudioCodec.cpp
@@ -423,11 +423,13 @@ namespace Audio
 	{
 		if (Requirements.DownstreamFormat == EBitRepresentation::Float32_Interleaved)
 		{
-			return MakeUnique<TDecoderOutputOwnBuffer<float>>(Requirements);
+			return MakeUnique<TCircularOutputBuffer<float>>(Requirements);
+			//return MakeUnique<TDecoderOutputOwnBuffer<float>>(Requirements);
 		}
 		else if (Requirements.DownstreamFormat == EBitRepresentation::Int16_Interleaved)
 		{
-			return MakeUnique<TDecoderOutputOwnBuffer<int16>>(Requirements);
+			return MakeUnique<TCircularOutputBuffer<int16>>(Requirements);
+			//return MakeUnique<TDecoderOutputOwnBuffer<int16>>(Requirements);
 		}
 		checkNoEntry();
 		return nullptr;
diff --git a/Engine/Source/Runtime/AudioExtensions/Private/IAudioCodecRegistry.cpp b/Engine/Source/Runtime/AudioExtensions/Private/IAudioCodecRegistry.cpp
--- a/Engine/Source/Runtime/AudioExtensions/Private/IAudioCodecRegistry.cpp
+++ b/Engine/Source/Runtime/AudioExtensions/Private/IAudioCodecRegistry.cpp
@@ -98,17 +98,23 @@ namespace Audio
 
 			auto FindLambda = [InCodecPtr](const TUniquePtr<ICodec>& i) -> bool { return i.Get() == InCodecPtr; };
 
+			auto RemoveLamda = [InCodecPtr](TMap<FName,TArray<FCodecPtr>>& InMap, FName InMapKey) -> void
+			{			
+				TArray<FCodecPtr>& VersionArray = InMap.FindChecked(InMapKey);
+				VersionArray.Remove(InCodecPtr);
+				if (VersionArray.Num() == 0)
+				{
+					InMap.Remove(InMapKey);
+				}
+			};
+
 			if (audio_ensure(LifetimeArray.ContainsByPredicate(FindLambda)))
 			{
 				UE_LOG(LogAudioCodec, Log, TEXT("Unregistered %s"), *InCodecPtr->GetDetails().ToString());
 
-				// Remove from lookups first. (maintain order).
-				TArray<FCodecPtr>& VersionArray = NameLookupMap.FindChecked(InCodecPtr->GetName());
-				VersionArray.Remove(InCodecPtr);
-				if( VersionArray.Num() == 0 )
-				{
-					NameLookupMap.Remove(InCodecPtr->GetName());
-				}
+				// Remove from Look ups.
+				RemoveLamda(FamilyNameLookupMap, InCodecPtr->GetDetails().FamilyName);
+				RemoveLamda(NameLookupMap, InCodecPtr->GetDetails().Name);
 
 				// Kill the object.
 				audio_ensure(LifetimeArray.RemoveAll(FindLambda) == 1);
@@ -254,13 +260,17 @@ namespace Audio
 
 	};
 
+	// Static instance.
+	TUniquePtr<Audio::ICodecRegistry> ICodecRegistry::Instance;
+
 	ICodecRegistry& ICodecRegistry::Get()
 	{
-		// FIXME.
-		//static IModuleInterface* ExtentionsModule = FModuleManager::LoadModulePtr<IModuleInterface>(FName(TEXT("AudioExtensions")));
-
-		static FCodecRegistry Impl;
-		return Impl;
+		// FIXME: Protect against race for the construction of the singleton.
+		if (!Instance.IsValid())
+		{
+			Instance = MakeUnique<FCodecRegistry>();
+		}
+		return *Instance;
 	}
 
 } //namespace Audio
diff --git a/Engine/Source/Runtime/AudioExtensions/Public/IAudioCodec.h b/Engine/Source/Runtime/AudioExtensions/Public/IAudioCodec.h
--- a/Engine/Source/Runtime/AudioExtensions/Public/IAudioCodec.h
+++ b/Engine/Source/Runtime/AudioExtensions/Public/IAudioCodec.h
@@ -77,7 +77,25 @@ namespace Audio
 	private:
 		uint32 Value = 0;
 	};
-	
+
+	// TODO: Move this somewhere central.
+	template<typename From, typename To>
+	static constexpr float GetConvertScalar()
+	{
+		if (TIsSame<From, To>::Value)
+		{
+			return 1.0f;
+		}
+		else if (TIsSame<From, int16>::Value && TIsSame<To, float>::Value)
+		{
+			return 1.f / (static_cast<float>(TNumericLimits<int16>::Max()) + 1.f);
+		}
+		else if (TIsSame<From, float>::Value && TIsSame<To, int16>::Value)
+		{
+			return static_cast<float>(TNumericLimits<int16>::Max()) + 1.f;
+		}
+	}
+		
 	enum EBitRepresentation
 	{
 		Int16_Interleaved,
@@ -174,13 +192,28 @@ namespace Audio
 			TArrayView<const float> InFloat32Interleave) = 0;
 
 		virtual int32 PopAudio( 
-			TArrayView<int16>& OutAudio ) = 0;
+			TArrayView<int16> InExternalBuffer,
+			FPushedAudioDetails& OutDetails ) = 0;
 
 		virtual int32 PopAudio(
-			TArrayView<float>& OutAudio ) = 0;
+			TArrayView<float> OutExternalBuffer,
+			FPushedAudioDetails& OutDetails ) = 0;		
 
-		virtual int32 Tell() const = 0;
+		// Convenience convert TArray -> TArrayView
+		int32 PopAudio(
+			TArray<int16>& InExternalBuffer,
+			FPushedAudioDetails& OutDetails)
+		{
+			return PopAudio(MakeArrayView(InExternalBuffer.GetData(), InExternalBuffer.Num()), OutDetails);
+		}
+		int32 PopAudio(
+			TArray<float>& InExternalBuffer,
+			FPushedAudioDetails& OutDetails)
+		{
+			return PopAudio(MakeArrayView(InExternalBuffer.GetData(), InExternalBuffer.Num()), OutDetails);
+		}
 	};
+		
 
 	// Decode input wrappers.
 
@@ -432,24 +465,7 @@ namespace Audio
 		{
 			return Reqs;
 		}	
-
-		template<typename From, typename To>
-		static constexpr float GetConvertScalar()
-		{
-			if( TIsSame<From,To>::Value ) 
-			{
-				return 1.0f;
-			}
-			else if( TIsSame<From,int16>::Value && TIsSame<To,float>::Value)
-			{
-				return 1.f / (static_cast<float>(TNumericLimits<int16>::Max()) + 1.f);
-			}
-			else if( TIsSame<From,float>::Value && TIsSame<To,int16>::Value)
-			{
-				return static_cast<float>(TNumericLimits<int16>::Max()) + 1.f;
-			}
-		}
-
+		
 		template<typename T> int32 PushAudioInternal(
 			const FPushedAudioDetails& InDetails,
 			TArrayView<const T> InAudio) 
@@ -491,41 +507,37 @@ namespace Audio
 			return PushAudioInternal(InDetails,In16BitInterleave);
 		}
 	
+		// This a bit naive, but ok for proof of concept.
 		template<typename T> int32 PopAudioInternal(
-			TArrayView<T>& OutSamples )
+			TArrayView<T> InExternalBuffer )
 		{
+			check(InExternalBuffer.Num() > 0);
 			if (TIsSame<T, TSampleType>::Value && Offset > 0)
 			{
 				int32 NumSamplesPopped = Offset;
 				Offset = 0;
+				NumSamplesPopped = FMath::Min(InExternalBuffer.Num(), NumSamplesPopped);
 				TArrayView<TSampleType> AV = Buffer.Slice(0, NumSamplesPopped);
-				OutSamples = MakeArrayView((T*)AV.GetData(), AV.Num());
+				FMemory::Memcpy(InExternalBuffer.GetData(), Buffer.GetData(), sizeof(T)* NumSamplesPopped);
 				return NumSamplesPopped;
 			}
-			
 			audio_ensure(false);
-			OutSamples = TArrayView<T>();
-			
 			return 0;
 		}
 
-		virtual int32 PopAudio(
-			TArrayView<float>& OutFloat )
-		{
-			return PopAudioInternal(OutFloat);
-		}
-
-		virtual int32 PopAudio(			
-			TArrayView<int16>& OutInt16 )
+		int32 PopAudio(
+			TArrayView<float> InExternalFloat32Buffer,
+			FPushedAudioDetails& OutDetails ) override
 		{
-			return PopAudioInternal(OutInt16);
+			return PopAudioInternal(InExternalFloat32Buffer);
 		}
 
-		// For back compatibility only.
-		int32 Tell() const 
+		int32 PopAudio(
+			TArrayView<int16> InExternalInt16Buffer,
+			FPushedAudioDetails& OutDetails ) override
 		{
-			return Offset;
-		}
+			return PopAudioInternal(InExternalInt16Buffer);
+		}	
 	};
 
 	template<typename T>
@@ -540,6 +552,78 @@ namespace Audio
 			TDecoderOutputArrayView<T>::SetArrayView(MakeArrayView(AlignedMemory));
 		}
 	};
+
+	template<typename TSampleType>
+	class TCircularOutputBuffer : public Audio::IDecoderOutput
+	{
+		Audio::TCircularAudioBuffer<TSampleType> Buffer;
+		FRequirements Reqs;
+	public:
+		TCircularOutputBuffer(
+			const FRequirements& InReqs)
+			: Buffer(InReqs.NumSampleFramesWanted*2) // Double buffered for now.
+			, Reqs(InReqs)
+		{
+		}
+
+		FRequirements GetRequirements(const FDecodedFormatInfo& InFormat) const override
+		{
+			return Reqs;
+		}
+
+		template<typename T>
+		int32 PushAudioInternal(
+			const FPushedAudioDetails& InDetails, TArrayView<const T> InBuffer) 
+		{
+			if (TIsSame<T, TSampleType>::Value)
+			{
+				return Buffer.Push(reinterpret_cast<const TSampleType*>(InBuffer.GetData()), InBuffer.Num());
+			}
+			else
+			{
+				constexpr float Scalar = GetConvertScalar<T, TSampleType>();
+				const T* Src = InBuffer.GetData();
+				int32 Num = InBuffer.Num();
+
+				// Convert 1 sample at a time, slow.
+				for (int32 i = 0; i < Num; ++i)
+				{
+					TSampleType Val = static_cast<TSampleType>(static_cast<float>(*Src++ * Scalar));
+					Buffer.Push(&Val, 1);
+				}
+
+				return Num;
+			}
+		}
+
+		int32 PushAudio(const FPushedAudioDetails& InDetails, TArrayView<const int16> In16BitInterleave) override
+		{
+			return PushAudioInternal(InDetails, In16BitInterleave);
+		}
+		int32 PushAudio(const FPushedAudioDetails& InDetails, TArrayView<const float> InFloat32Interleave) override
+		{
+			return PushAudioInternal(InDetails, InFloat32Interleave);
+		}
+		int32 PopAudio(TArrayView<int16> InExternalInt16Buffer, FPushedAudioDetails& OutDetails) override
+		{
+			if( TIsSame<TSampleType, int16>::Value )
+			{
+				return Buffer.Pop(reinterpret_cast<TSampleType*>(InExternalInt16Buffer.GetData()), InExternalInt16Buffer.Num());
+			}
+			audio_ensure(false);
+			return 0;
+		}
+
+		int32 PopAudio(TArrayView<float> InExternalFloat32Buffer, FPushedAudioDetails& OutDetails) override
+		{
+			if( TIsSame<TSampleType, float>::Value )
+			{
+				return Buffer.Pop(reinterpret_cast<TSampleType*>(InExternalFloat32Buffer.GetData()), InExternalFloat32Buffer.Num());
+			}
+			audio_ensure(false);
+			return 0;			
+		}
+	};
 };
 
 #pragma endregion Decoder
diff --git a/Engine/Source/Runtime/AudioExtensions/Public/IAudioCodecRegistry.h b/Engine/Source/Runtime/AudioExtensions/Public/IAudioCodecRegistry.h
--- a/Engine/Source/Runtime/AudioExtensions/Public/IAudioCodecRegistry.h
+++ b/Engine/Source/Runtime/AudioExtensions/Public/IAudioCodecRegistry.h
@@ -11,12 +11,15 @@ namespace Audio
 {
 	class AUDIOEXTENSIONS_API ICodecRegistry
 	{
+		static TUniquePtr<ICodecRegistry> Instance;
 	public:
 		using FCodecPtr = ICodec * ;
-	
+
 		// Abstract Singleton.
 		static ICodecRegistry& Get();
 
+		virtual ~ICodecRegistry() = default;
+
 		// Register each codec.
 		virtual bool RegisterCodec(TUniquePtr<ICodec>&&) = 0;
 		virtual bool UnregisterCodec(FCodecPtr) = 0;
