Renaming FindCodecByFromParsingInput to FindCodecByParsingInput

diff --git a/Engine/Plugins/Runtime/Metasound/Source/MetasoundEngine/Private/MetasoundWavePlayerNode.cpp b/Engine/Plugins/Runtime/Metasound/Source/MetasoundEngine/Private/MetasoundWavePlayerNode.cpp
--- a/Engine/Plugins/Runtime/Metasound/Source/MetasoundEngine/Private/MetasoundWavePlayerNode.cpp
+++ b/Engine/Plugins/Runtime/Metasound/Source/MetasoundEngine/Private/MetasoundWavePlayerNode.cpp
@@ -139,7 +139,7 @@ namespace Metasound
 		FWave::FDecoderInputPtr Input = FWave::CreateDecoderInput(Wave);
 		if (Input)
 		{
-			ICodecRegistry::FCodecPtr Codec = ICodecRegistry::Get().FindCodecByFromParsingInput(Input.Get());
+			ICodecRegistry::FCodecPtr Codec = ICodecRegistry::Get().FindCodecByParsingInput(Input.Get());
 			if (Codec)
 			{
 				// V1, Ask for an output buffer the size of a frame.
diff --git a/Engine/Source/Runtime/AudioExtensions/Private/IAudioCodecRegistry.cpp b/Engine/Source/Runtime/AudioExtensions/Private/IAudioCodecRegistry.cpp
--- a/Engine/Source/Runtime/AudioExtensions/Private/IAudioCodecRegistry.cpp
+++ b/Engine/Source/Runtime/AudioExtensions/Private/IAudioCodecRegistry.cpp
@@ -187,7 +187,7 @@ namespace Audio
 			return nullptr;
 		}
 		
-		virtual ICodecRegistry::FCodecPtr FindCodecByFromParsingInput(
+		virtual ICodecRegistry::FCodecPtr FindCodecByParsingInput(
 			IDecoderInput* InObject) const override
 		{
 			if( !audio_ensure(InObject))
diff --git a/Engine/Source/Runtime/AudioExtensions/Private/IAudioDecoderScheduler.cpp b/Engine/Source/Runtime/AudioExtensions/Private/IAudioDecoderScheduler.cpp
--- a/Engine/Source/Runtime/AudioExtensions/Private/IAudioDecoderScheduler.cpp
+++ b/Engine/Source/Runtime/AudioExtensions/Private/IAudioDecoderScheduler.cpp
@@ -69,7 +69,7 @@ namespace Audio
 					// fall-through.
 					case eCreatingDecoder:
 					{
-						ICodec* pCodec = ICodecRegistry::Get().FindCodecByFromParsingInput(
+						ICodec* pCodec = ICodecRegistry::Get().FindCodecByParsingInput(
 							StartArgs.InputObject.Get());
 
 						if( !pCodec )
diff --git a/Engine/Source/Runtime/AudioExtensions/Public/IAudioCodecRegistry.h b/Engine/Source/Runtime/AudioExtensions/Public/IAudioCodecRegistry.h
--- a/Engine/Source/Runtime/AudioExtensions/Public/IAudioCodecRegistry.h
+++ b/Engine/Source/Runtime/AudioExtensions/Public/IAudioCodecRegistry.h
@@ -28,7 +28,7 @@ namespace Audio
 		virtual FCodecPtr FindCodecByName(FName InName, int32 InVersion = INDEX_NONE) const = 0;
 
 		// Find a codec given a decoders input interface
-		virtual FCodecPtr FindCodecByFromParsingInput(IDecoderInput* InObject) const = 0;
+		virtual FCodecPtr FindCodecByParsingInput(IDecoderInput* InObject) const = 0;
 
 		// Find codec by finding by family name.  (Version can be INDEX_NONE and give the latest version).
 		virtual FCodecPtr FindCodecByFamilyName(FName InFamilyName, int32 InVersion = INDEX_NONE) const = 0;
