Missing file

diff --git a/Engine/Source/Runtime/AudioExtensions/Public/IAudioModulation.h b/Engine/Source/Runtime/AudioExtensions/Public/IAudioModulation.h
--- a/Engine/Source/Runtime/AudioExtensions/Public/IAudioModulation.h
+++ b/Engine/Source/Runtime/AudioExtensions/Public/IAudioModulation.h
@@ -30,9 +30,9 @@ namespace Audio
 	using FModulatorTypeId = uint32;
 	using FModulatorHandleId = uint32;
 
-	using FModulationUnitConvertFunction = TFunction<void(float* RESTRICT /* OutValueLinearToUnitBuffer */, int32 /* InNumSamples */)>;
-	using FModulationLinearConversionFunction = TFunction<void(float* RESTRICT /* OutValueUnitToLinearBuffer */, int32 /* InNumSamples */)>;
-	using FModulationMixFunction = TFunction<void(float* RESTRICT /* OutBufferLinearA */, const float* RESTRICT /* InBufferLinearB */, int32 /* InNumSamples */)>;
+	using FModulationUnitConvertFunction = TFunction<void(float* RESTRICT /* OutValueNormalizedToUnitBuffer */, int32 /* InNumSamples */)>;
+	using FModulationNormalizedConversionFunction = TFunction<void(float* RESTRICT /* OutValueUnitToNormalizedBuffer */, int32 /* InNumSamples */)>;
+	using FModulationMixFunction = TFunction<void(float* RESTRICT /* OutBufferNormalizedA */, const float* RESTRICT /* InBufferNormalizedB */, int32 /* InNumSamples */)>;
 
 	struct AUDIOEXTENSIONS_API FModulationParameter
 	{
@@ -56,15 +56,15 @@ namespace Audio
 		FText UnitDisplayName;
 #endif // WITH_EDITORONLY_DATA
 
-		// Function used to convert value buffer from linear space [0.0f, 1.0f] to unit space.
+		// Function used to convert value buffer from normalized, unitless space [0.0f, 1.0f] to unit space.
 		FModulationUnitConvertFunction UnitFunction;
 
-		// Function used to convert value buffer from unit space to linear [0.0f, 1.0f] space.
-		FModulationLinearConversionFunction LinearFunction;
+		// Function used to convert value buffer from unit space to normalized, unitless [0.0f, 1.0f] space.
+		FModulationNormalizedConversionFunction NormalizedFunction;
 
 		static const FModulationMixFunction& GetDefaultMixFunction();
 
-		// Function used to mix linear values together.
+		// Function used to mix normalized values together.
 		FModulationMixFunction MixFunction;
 	};
 
