Fixing PVS 7.7 Issues:

warning V779: Unreachable code detected. It is possible that an error is present.

diff --git a/Engine/Source/Runtime/AudioExtensions/Private/IAudioCodec.cpp b/Engine/Source/Runtime/AudioExtensions/Private/IAudioCodec.cpp
--- a/Engine/Source/Runtime/AudioExtensions/Private/IAudioCodec.cpp
+++ b/Engine/Source/Runtime/AudioExtensions/Private/IAudioCodec.cpp
@@ -431,7 +431,7 @@ namespace Audio
 			return MakeUnique<TCircularOutputBuffer<int16>>(Requirements);
 			//return MakeUnique<TDecoderOutputOwnBuffer<int16>>(Requirements);
 		}
-		checkNoEntry();
+		checkNoEntry(); //-V779
 		return nullptr;
 	}
 
@@ -451,7 +451,7 @@ namespace Audio
 			Output->SetArrayViewBytes(InExternalMemory);
 			return Output;
 		}
-		checkNoEntry();
+		checkNoEntry(); //-V779
 		return nullptr;
 	}
 
