- Add ability to send MetaSound object parameter types that are not SoundWaves
- Cache Audio Parameters on AudioComponent properly so they can be modified prior to playback.
- Migrate AudioComponentParameterization to SoundGeneratorParameterInterface
- AudioComponent now inherits from SoundGeneratorParameterInterface
- Unify code path with default transmitter implemented in SoundBase for legacy SoundCue parameter system

Other Minor updates:
- Rename to "Audio Instance" to "Audio Parameter" for clarity

diff --git a/Engine/Config/BaseEngine.ini b/Engine/Config/BaseEngine.ini
--- a/Engine/Config/BaseEngine.ini
+++ b/Engine/Config/BaseEngine.ini
@@ -925,6 +925,11 @@ NearClipPlane=10.0
 +ClassRedirects=(OldName="/Script/MovieSceneTracks.MovieSceneVectorTrackExtensions",NewName="/Script/SequencerScripting.MovieSceneFloatVectorTrackExtensions")
 +StructRedirects=(OldName="/Script/MovieSceneTracks.MovieSceneVectorKeyStructBase",NewName="/Script/MovieSceneTracks.MovieSceneFloatVectorKeyStructBase")
 +StructRedirects=(OldName="/Script/MovieSceneTracks.MovieSceneVectorKeyStruct",NewName="/Script/MovieSceneTracks.MovieSceneVector3fKeyStruct")
++StructRedirects=(OldName="AudioComponentParam", NewName="AudioParameter")
++PropertyRedirects=(OldName="AudioParameter.SoundWaveParam", NewName="AudioParameter.ObjectParam")
++FunctionRedirects=(OldName="AudioComponent.SetBoolParameter", NewName="SoundGeneratorParameterInterface.SetBoolParameter")
++FunctionRedirects=(OldName="AudioComponent.SetFloatParameter", NewName="SoundGeneratorParameterInterface.SetFloatParameter")
++FunctionRedirects=(OldName="AudioComponent.SetIntParameter", NewName="SoundGeneratorParameterInterface.SetIntParameter")
 
 
 [CoreUObject.Metadata]
diff --git a/Engine/Plugins/Runtime/Metasound/Source/MetasoundEditor/Private/MetasoundEditor.cpp b/Engine/Plugins/Runtime/Metasound/Source/MetasoundEditor/Private/MetasoundEditor.cpp
--- a/Engine/Plugins/Runtime/Metasound/Source/MetasoundEditor/Private/MetasoundEditor.cpp
+++ b/Engine/Plugins/Runtime/Metasound/Source/MetasoundEditor/Private/MetasoundEditor.cpp
@@ -823,7 +823,7 @@ namespace Metasound
 				{
 					PlayTime = 0.0;
 
-					UObject* ParamInterfaceObject = PreviewComp->GetParameterInterface().GetObject();
+					UObject* ParamInterfaceObject = PreviewComp;
 					if (ensure(ParamInterfaceObject))
 					{
 						SetPreviewID(ParamInterfaceObject->GetUniqueID());
@@ -947,11 +947,11 @@ namespace Metasound
 					{
 						// TODO: fix how identifying the parameter to update is determined. It should not be done
 						// with a "DisplayName" but rather the vertex Guid.
-						if (TScriptInterface<IAudioParameterInterface> ParamInterface = PreviewComponent->GetParameterInterface())
+						if (TScriptInterface<IAudioParameterInterface> ParamInterface = PreviewComponent)
 						{
 							FConstNodeHandle NodeHandle = InputNode->GetConstNodeHandle();
 							FVertexKey VertexKey = Metasound::FVertexKey(NodeHandle->GetDisplayName().ToString());
-							ParamInterface->Trigger(*VertexKey);
+							ParamInterface->SetTriggerParameter(*VertexKey);
 						}
 					}
 				}
diff --git a/Engine/Plugins/Runtime/Metasound/Source/MetasoundEditor/Private/MetasoundEditorGraph.cpp b/Engine/Plugins/Runtime/Metasound/Source/MetasoundEditor/Private/MetasoundEditorGraph.cpp
--- a/Engine/Plugins/Runtime/Metasound/Source/MetasoundEditor/Private/MetasoundEditorGraph.cpp
+++ b/Engine/Plugins/Runtime/Metasound/Source/MetasoundEditor/Private/MetasoundEditorGraph.cpp
@@ -1,7 +1,7 @@
 // Copyright Epic Games, Inc. All Rights Reserved.
 #include "MetasoundEditorGraph.h"
 
-#include "Audio/AudioParameterInterface.h"
+#include "AudioParameterInterface.h"
 #include "Components/AudioComponent.h"
 #include "EdGraph/EdGraphNode.h"
 #include "Interfaces/ITargetPlatform.h"
@@ -311,7 +311,7 @@ void UMetasoundEditorGraphInput::OnLiteralChanged(bool bPostTransaction)
 		UAudioComponent* PreviewComponent = GEditor->GetPreviewAudioComponent();
 		check(PreviewComponent);
 
-		if (TScriptInterface<IAudioParameterInterface> ParamInterface = PreviewComponent->GetParameterInterface())
+		if (TScriptInterface<IAudioParameterInterface> ParamInterface = PreviewComponent)
 		{
 			// TODO: fix how identifying the parameter to update is determined. It should not be done
 			// with a "DisplayName" but rather the vertex Guid.
@@ -609,13 +609,7 @@ bool UMetasoundEditorGraph::IsPreviewing() const
 		return false;
 	}
 
-	UObject* ParamInterfaceObject = PreviewComponent->GetParameterInterface().GetObject();
-	if (!ParamInterfaceObject)
-	{
-		return false;
-	}
-
-	return ParamInterfaceObject->GetUniqueID() == PreviewID;
+	return PreviewComponent->GetUniqueID() == PreviewID;
 }
 
 void UMetasoundEditorGraph::IterateOutputs(TUniqueFunction<void(UMetasoundEditorGraphOutput&)> InFunction) const
diff --git a/Engine/Plugins/Runtime/Metasound/Source/MetasoundEditor/Private/MetasoundEditorGraphInputNodes.cpp b/Engine/Plugins/Runtime/Metasound/Source/MetasoundEditor/Private/MetasoundEditorGraphInputNodes.cpp
--- a/Engine/Plugins/Runtime/Metasound/Source/MetasoundEditor/Private/MetasoundEditorGraphInputNodes.cpp
+++ b/Engine/Plugins/Runtime/Metasound/Source/MetasoundEditor/Private/MetasoundEditorGraphInputNodes.cpp
@@ -3,7 +3,7 @@
 
 #include "Algo/Count.h"
 #include "Algo/Transform.h"
-#include "Audio/AudioParameterInterface.h"
+#include "AudioParameterInterface.h"
 #include "CoreMinimal.h"
 #include "EdGraph/EdGraphNode.h"
 #include "GraphEditorSettings.h"
@@ -126,7 +126,7 @@ void UMetasoundEditorGraphInputBool::SetFromLiteral(const FMetasoundFrontendLite
 
 void UMetasoundEditorGraphInputBool::UpdatePreviewInstance(const Metasound::FVertexKey& InParameterName, TScriptInterface<IAudioParameterInterface>& InParameterInterface) const
 	{
-		InParameterInterface->SetBool(*InParameterName, Default.Value);
+		InParameterInterface->SetBoolParameter(*InParameterName, Default.Value);
 	}
 
 FMetasoundFrontendLiteral UMetasoundEditorGraphInputBoolArray::GetDefault() const
@@ -153,7 +153,7 @@ void UMetasoundEditorGraphInputBoolArray::UpdatePreviewInstance(const Metasound:
 {
 	TArray<bool> BoolArray;
 	Algo::Transform(Default, BoolArray, [](const FMetasoundEditorGraphInputBoolRef& InValue) { return InValue.Value; });
-	InParameterInterface->SetBoolArray(*InParameterName, BoolArray);
+	InParameterInterface->SetBoolArrayParameter(*InParameterName, BoolArray);
 }
 
 FMetasoundFrontendLiteral UMetasoundEditorGraphInputInt::GetDefault() const
@@ -175,7 +175,7 @@ void UMetasoundEditorGraphInputInt::SetFromLiteral(const FMetasoundFrontendLiter
 
 void UMetasoundEditorGraphInputInt::UpdatePreviewInstance(const Metasound::FVertexKey& InParameterName, TScriptInterface<IAudioParameterInterface>& InParameterInterface) const
 {
-	InParameterInterface->SetInt(*InParameterName, Default.Value);
+	InParameterInterface->SetIntParameter(*InParameterName, Default.Value);
 }
 
 FMetasoundFrontendLiteral UMetasoundEditorGraphInputIntArray::GetDefault() const
@@ -203,7 +203,7 @@ void UMetasoundEditorGraphInputIntArray::UpdatePreviewInstance(const Metasound::
 {
 	TArray<int32> IntArray;
 	Algo::Transform(Default, IntArray, [](const FMetasoundEditorGraphInputIntRef& InValue) { return InValue.Value; });
-	InParameterInterface->SetIntArray(*InParameterName, IntArray);
+	InParameterInterface->SetIntArrayParameter(*InParameterName, IntArray);
 }
 
 FMetasoundFrontendLiteral UMetasoundEditorGraphInputFloat::GetDefault() const
@@ -225,7 +225,7 @@ void UMetasoundEditorGraphInputFloat::SetFromLiteral(const FMetasoundFrontendLit
 
 void UMetasoundEditorGraphInputFloat::UpdatePreviewInstance(const Metasound::FVertexKey& InParameterName, TScriptInterface<IAudioParameterInterface>& InParameterInterface) const
 {
-	InParameterInterface->SetFloat(*InParameterName, Default);
+	InParameterInterface->SetFloatParameter(*InParameterName, Default);
 }
 
 FMetasoundFrontendLiteral UMetasoundEditorGraphInputFloatArray::GetDefault() const
@@ -247,7 +247,7 @@ void UMetasoundEditorGraphInputFloatArray::SetFromLiteral(const FMetasoundFronte
 
 void UMetasoundEditorGraphInputFloatArray::UpdatePreviewInstance(const Metasound::FVertexKey& InParameterName, TScriptInterface<IAudioParameterInterface>& InParameterInterface) const
 {
-	InParameterInterface->SetFloatArray(*InParameterName, Default);
+	InParameterInterface->SetFloatArrayParameter(*InParameterName, Default);
 }
 
 FMetasoundFrontendLiteral UMetasoundEditorGraphInputString::GetDefault() const
@@ -269,7 +269,7 @@ void UMetasoundEditorGraphInputString::SetFromLiteral(const FMetasoundFrontendLi
 
 void UMetasoundEditorGraphInputString::UpdatePreviewInstance(const Metasound::FVertexKey& InParameterName, TScriptInterface<IAudioParameterInterface>& InParameterInterface) const
 {
-	InParameterInterface->SetString(*InParameterName, Default);
+	InParameterInterface->SetStringParameter(*InParameterName, Default);
 }
 
 FMetasoundFrontendLiteral UMetasoundEditorGraphInputStringArray::GetDefault() const
@@ -291,7 +291,7 @@ void UMetasoundEditorGraphInputStringArray::SetFromLiteral(const FMetasoundFront
 
 void UMetasoundEditorGraphInputStringArray::UpdatePreviewInstance(const Metasound::FVertexKey& InParameterName, TScriptInterface<IAudioParameterInterface>& InParameterInterface) const
 {
-	InParameterInterface->SetStringArray(*InParameterName, Default);
+	InParameterInterface->SetStringArrayParameter(*InParameterName, Default);
 }
 
 FMetasoundFrontendLiteral UMetasoundEditorGraphInputObject::GetDefault() const
diff --git a/Engine/Plugins/Runtime/Metasound/Source/MetasoundEditor/Private/SMetasoundGraphNode.cpp b/Engine/Plugins/Runtime/Metasound/Source/MetasoundEditor/Private/SMetasoundGraphNode.cpp
--- a/Engine/Plugins/Runtime/Metasound/Source/MetasoundEditor/Private/SMetasoundGraphNode.cpp
+++ b/Engine/Plugins/Runtime/Metasound/Source/MetasoundEditor/Private/SMetasoundGraphNode.cpp
@@ -1,7 +1,7 @@
 // Copyright Epic Games, Inc. All Rights Reserved.
 #include "SMetasoundGraphNode.h"
 
-#include "Audio/AudioParameterInterface.h"
+#include "AudioParameterInterface.h"
 #include "Components/AudioComponent.h"
 #include "EditorStyleSet.h"
 #include "GraphEditorSettings.h"
@@ -71,14 +71,11 @@ void SMetasoundGraphNode::ExecuteInputTrigger(UMetasoundEditorGraphInputLiteral&
 
 	if (UAudioComponent* PreviewComponent = GEditor->GetPreviewAudioComponent())
 	{
-		if (TScriptInterface<IAudioParameterInterface> ParamInterface = PreviewComponent->GetParameterInterface())
-		{
-			// TODO: fix how identifying the parameter to update is determined. It should not be done
-			// with a "DisplayName" but rather the vertex Guid.
-			Metasound::Frontend::FConstNodeHandle NodeHandle = Input->GetConstNodeHandle();
-			Metasound::FVertexKey VertexKey = Metasound::FVertexKey(NodeHandle->GetDisplayName().ToString());
-			ParamInterface->Trigger(*VertexKey);
-		}
+		// TODO: fix how identifying the parameter to update is determined. It should not be done
+		// with a "DisplayName" but rather the vertex Guid.
+		Metasound::Frontend::FConstNodeHandle NodeHandle = Input->GetConstNodeHandle();
+		Metasound::FVertexKey VertexKey = Metasound::FVertexKey(NodeHandle->GetDisplayName().ToString());
+		PreviewComponent->SetTriggerParameter(*VertexKey);
 	}
 }
 
diff --git a/Engine/Plugins/Runtime/Metasound/Source/MetasoundEditor/Public/MetasoundEditorGraph.h b/Engine/Plugins/Runtime/Metasound/Source/MetasoundEditor/Public/MetasoundEditorGraph.h
--- a/Engine/Plugins/Runtime/Metasound/Source/MetasoundEditor/Public/MetasoundEditorGraph.h
+++ b/Engine/Plugins/Runtime/Metasound/Source/MetasoundEditor/Public/MetasoundEditorGraph.h
@@ -1,8 +1,9 @@
 // Copyright Epic Games, Inc. All Rights Reserved.
 #pragma once
 
-#include "Audio/AudioParameterInterface.h"
 #include "CoreMinimal.h"
+
+#include "AudioParameterInterface.h"
 #include "MetasoundEditor.h"
 #include "MetasoundEditorGraphNode.h"
 #include "MetasoundFrontendController.h"
@@ -16,12 +17,14 @@
 
 #include "MetasoundEditorGraph.generated.h"
 
+
 // Forward Declarations
 struct FMetasoundFrontendDocument;
 class ITargetPlatform;
 class UMetasoundEditorGraphInputNode;
 class UMetasoundEditorGraphNode;
 
+
 namespace Metasound
 {
 	namespace Editor
diff --git a/Engine/Plugins/Runtime/Metasound/Source/MetasoundEditor/Public/MetasoundEditorGraphInputNodes.h b/Engine/Plugins/Runtime/Metasound/Source/MetasoundEditor/Public/MetasoundEditorGraphInputNodes.h
--- a/Engine/Plugins/Runtime/Metasound/Source/MetasoundEditor/Public/MetasoundEditorGraphInputNodes.h
+++ b/Engine/Plugins/Runtime/Metasound/Source/MetasoundEditor/Public/MetasoundEditorGraphInputNodes.h
@@ -3,7 +3,7 @@
 
 #include "Algo/Count.h"
 #include "Algo/Transform.h"
-#include "Audio/AudioParameterInterface.h"
+#include "AudioParameterInterface.h"
 #include "CoreMinimal.h"
 #include "EdGraph/EdGraphNode.h"
 #include "GraphEditorSettings.h"
diff --git a/Engine/Plugins/Runtime/Metasound/Source/MetasoundEngine/Private/Metasound.cpp b/Engine/Plugins/Runtime/Metasound/Source/MetasoundEngine/Private/Metasound.cpp
--- a/Engine/Plugins/Runtime/Metasound/Source/MetasoundEngine/Private/Metasound.cpp
+++ b/Engine/Plugins/Runtime/Metasound/Source/MetasoundEngine/Private/Metasound.cpp
@@ -15,9 +15,9 @@
 #include "MetasoundFrontendRegistries.h"
 #include "MetasoundFrontendSearchEngine.h"
 #include "MetasoundGenerator.h"
-#include "MetasoundInstanceTransmitter.h"
 #include "MetasoundLog.h"
 #include "MetasoundOperatorSettings.h"
+#include "MetasoundParameterTransmitter.h"
 #include "MetasoundPrimitives.h"
 #include "MetasoundReceiveNode.h"
 #include "MetasoundTrigger.h"
diff --git a/Engine/Plugins/Runtime/Metasound/Source/MetasoundEngine/Private/MetasoundEngineArchetypes.cpp b/Engine/Plugins/Runtime/Metasound/Source/MetasoundEngine/Private/MetasoundEngineArchetypes.cpp
--- a/Engine/Plugins/Runtime/Metasound/Source/MetasoundEngine/Private/MetasoundEngineArchetypes.cpp
+++ b/Engine/Plugins/Runtime/Metasound/Source/MetasoundEngine/Private/MetasoundEngineArchetypes.cpp
@@ -310,7 +310,7 @@ namespace Metasound
 
 			const FString& GetInstanceIDName()
 			{
-				static const FString& Name = FMetasoundInstanceTransmitter::GetInstanceIDEnvironmentVariableName();
+				static const FString& Name = FMetaSoundParameterTransmitter::GetInstanceIDEnvironmentVariableName();
 				return Name;
 			}
 
@@ -488,7 +488,7 @@ namespace Metasound
 
 			const FString& GetInstanceIDName()
 			{
-				static const FString& Name = FMetasoundInstanceTransmitter::GetInstanceIDEnvironmentVariableName();
+				static const FString& Name = FMetaSoundParameterTransmitter::GetInstanceIDEnvironmentVariableName();
 				return Name;
 			}
 
diff --git a/Engine/Plugins/Runtime/Metasound/Source/MetasoundEngine/Private/MetasoundSource.cpp b/Engine/Plugins/Runtime/Metasound/Source/MetasoundEngine/Private/MetasoundSource.cpp
--- a/Engine/Plugins/Runtime/Metasound/Source/MetasoundEngine/Private/MetasoundSource.cpp
+++ b/Engine/Plugins/Runtime/Metasound/Source/MetasoundEngine/Private/MetasoundSource.cpp
@@ -10,14 +10,15 @@
 #include "MetasoundEngineEnvironment.h"
 #include "MetasoundEnvironment.h"
 #include "MetasoundFrontendController.h"
+#include "MetasoundFrontendDataTypeRegistry.h"
 #include "MetasoundFrontendInjectReceiveNodes.h"
 #include "MetasoundFrontendQuery.h"
 #include "MetasoundFrontendQuerySteps.h"
 #include "MetasoundFrontendTransform.h"
 #include "MetasoundGenerator.h"
-#include "MetasoundInstanceTransmitter.h"
 #include "MetasoundLog.h"
 #include "MetasoundOperatorSettings.h"
+#include "MetasoundParameterTransmitter.h"
 #include "MetasoundPrimitives.h"
 #include "MetasoundReceiveNode.h"
 #include "MetasoundTrace.h"
@@ -44,14 +45,14 @@ namespace MetaSoundSourcePrivate
 		(
 			{
 				{
-					EMetasoundSourceAudioFormat::Mono, 
+					EMetasoundSourceAudioFormat::Mono,
 					{ MetasoundSourceMono::GetAudioOutputName() }
 				},
 				{
-					EMetasoundSourceAudioFormat::Stereo, 
-					{ 
-						MetasoundSourceStereo::GetLeftAudioOutputName(), 
-						MetasoundSourceStereo::GetRightAudioOutputName() 
+					EMetasoundSourceAudioFormat::Stereo,
+					{
+						MetasoundSourceStereo::GetLeftAudioOutputName(),
+						MetasoundSourceStereo::GetRightAudioOutputName()
 					}
 				}
 			}
@@ -236,9 +237,146 @@ void UMetaSoundSource::SetRegistryAssetClassInfo(const Metasound::Frontend::FNod
 }
 #endif // WITH_EDITORONLY_DATA
 
+void UMetaSoundSource::InitParameters(TArray<FAudioParameter>& InParametersToInit, FName InFeatureName)
+{
+	auto Sanitize = [](FAudioParameter& Parameter)
+	{
+		switch (Parameter.ParamType)
+		{
+			case EAudioParameterType::Boolean:
+			{
+				Parameter = FAudioParameter(Parameter.ParamName, Parameter.BoolParam);
+			}
+			break;
+
+			case EAudioParameterType::BooleanArray:
+			{
+				TArray<bool> TempArray = Parameter.ArrayBoolParam;
+				Parameter = FAudioParameter(Parameter.ParamName, MoveTemp(TempArray));
+			}
+			break;
+
+			case EAudioParameterType::Float:
+			{
+				Parameter = FAudioParameter(Parameter.ParamName, Parameter.FloatParam);
+			}
+			break;
+
+			case EAudioParameterType::FloatArray:
+			{
+				TArray<float> TempArray = Parameter.ArrayFloatParam;
+				Parameter = FAudioParameter(Parameter.ParamName, MoveTemp(TempArray));
+			}
+			break;
+
+			case EAudioParameterType::Integer:
+			{
+				Parameter = FAudioParameter(Parameter.ParamName, Parameter.IntParam);
+			}
+			break;
+
+			case EAudioParameterType::IntegerArray:
+			{
+				TArray<int32> TempArray = Parameter.ArrayIntParam;
+				Parameter = FAudioParameter(Parameter.ParamName, MoveTemp(TempArray));
+			}
+			break;
+
+			case EAudioParameterType::Object:
+			{
+				Parameter = FAudioParameter(Parameter.ParamName, Parameter.ObjectParam);
+			}
+			break;
+
+			case EAudioParameterType::ObjectArray:
+			{
+				TArray<UObject*> TempArray = Parameter.ArrayObjectParam;
+				Parameter = FAudioParameter(Parameter.ParamName, MoveTemp(TempArray));
+			}
+			break;
+
+
+			case EAudioParameterType::String:
+			{
+				Parameter = FAudioParameter(Parameter.ParamName, Parameter.StringParam);
+			}
+			break;
+
+			case EAudioParameterType::StringArray:
+			{
+				TArray<FString> TempArray = Parameter.ArrayStringParam;
+				Parameter = FAudioParameter(Parameter.ParamName, MoveTemp(TempArray));
+			}
+			break;
+
+			case EAudioParameterType::None:
+			default:
+			break;
+		}
+	};
+
+	auto ConstructProxies = [this, FeatureName = InFeatureName](FAudioParameter& OutParamToInit)
+	{
+		using namespace Metasound;
+		using namespace Metasound::Frontend;
+
+		const Audio::FProxyDataInitParams ProxyInitParams { FeatureName };
+
+		if (!IsParameterValid(OutParamToInit))
+		{
+			UE_LOG(LogMetaSound, Error, TEXT("Failed to set invalid parameter '%s': Either does not exist or is invalid type"), *OutParamToInit.ParamName.ToString());
+			return;
+		}
+
+		switch (OutParamToInit.ParamType)
+		{
+			case EAudioParameterType::Object:
+			{
+				FDataTypeRegistryInfo DataTypeInfo;
+				if (IDataTypeRegistry::Get().GetDataTypeInfo(OutParamToInit.ObjectParam, DataTypeInfo))
+				{
+					Audio::IProxyDataPtr ProxyPtr = IDataTypeRegistry::Get().CreateProxyFromUObject(DataTypeInfo.DataTypeName, OutParamToInit.ObjectParam);
+					OutParamToInit.ObjectProxies.Emplace(MoveTemp(ProxyPtr));
+
+					// Null out param as it is no longer needed (nor desired to be accessed once passed to the Audio Thread)
+					OutParamToInit.ObjectParam = nullptr;
+				}
+			}
+			break;
+
+			case EAudioParameterType::ObjectArray:
+			{
+				for (TObjectPtr<UObject>& Object : OutParamToInit.ArrayObjectParam)
+				{
+					FDataTypeRegistryInfo DataTypeInfo;
+					if (IDataTypeRegistry::Get().GetDataTypeInfo(Object, DataTypeInfo))
+					{
+						Audio::IProxyDataPtr ProxyPtr = IDataTypeRegistry::Get().CreateProxyFromUObject(DataTypeInfo.DataTypeName, Object);
+						OutParamToInit.ObjectProxies.Emplace(MoveTemp(ProxyPtr));
+					}
+				}
+				// Reset param array as it is no longer needed (nor desired to be accessed once passed to the Audio Thread).
+				// All object manipulation hereafter should be done via proxies
+				OutParamToInit.ArrayObjectParam.Reset();
+			}
+			break;
+
+			default:
+				break;
+		}
+	};
+
+	for (FAudioParameter& Parameter : InParametersToInit)
+	{
+		Sanitize(Parameter);
+		ConstructProxies(Parameter);
+	}
+}
+
 void UMetaSoundSource::InitResources()
 {
-	METASOUND_TRACE_CPUPROFILER_EVENT_SCOPE(UMetaSoundAssetSubsystem::AddOrUpdateAsset);
+	using namespace Metasound::Frontend;
+	METASOUND_TRACE_CPUPROFILER_EVENT_SCOPE(UMetaSoundSource::InitResources);
 
 	FMetaSoundAssetRegistrationOptions RegOptions;
 	RegOptions.bForceReregister = false;
@@ -312,16 +450,161 @@ ISoundGeneratorPtr UMetaSoundSource::CreateSoundGenerator(const FSoundGeneratorI
 	return ISoundGeneratorPtr(new FMetasoundGenerator(MoveTemp(InitParams)));
 }
 
-TUniquePtr<IAudioInstanceTransmitter> UMetaSoundSource::CreateInstanceTransmitter(const FAudioInstanceTransmitterInitParams& InParams) const
+bool UMetaSoundSource::IsParameterValid(const FAudioParameter& InParameter) const
+{
+	using namespace Metasound;
+	using namespace Metasound::Frontend;
+
+	bool bIsValid = false;
+
+	GetRootGraphHandle()->IterateConstNodes([Param = &InParameter, bValid = &bIsValid](FConstNodeHandle NodeHandle)
+	{
+		using namespace Metasound;
+		using namespace Metasound::Frontend;
+
+		// TODO: Fix this.  NodeHandles should use FNames, not display names (this is broken in other languages)
+		if (Param->ParamName == *NodeHandle->GetDisplayName().ToString())
+		{
+			TArray<FConstInputHandle> InputHandles = NodeHandle->GetConstInputs();
+			if (InputHandles.IsEmpty())
+			{
+				return;
+			}
+
+			const FName DataTypeName = InputHandles[0]->GetDataType();
+			switch (Param->ParamType)
+			{
+				case EAudioParameterType::Boolean:
+				{
+					FDataTypeRegistryInfo DataTypeInfo;
+					*bValid = IDataTypeRegistry::Get().GetDataTypeInfo(DataTypeName, DataTypeInfo);
+					*bValid &= DataTypeInfo.bIsBoolParsable;
+				}
+				break;
+
+				case EAudioParameterType::BooleanArray:
+				{
+					FDataTypeRegistryInfo DataTypeInfo;
+					*bValid = IDataTypeRegistry::Get().GetDataTypeInfo(DataTypeName, DataTypeInfo);
+					*bValid &= DataTypeInfo.bIsBoolArrayParsable;
+				}
+				break;
+
+				case EAudioParameterType::Float:
+				{
+					FDataTypeRegistryInfo DataTypeInfo;
+					*bValid = IDataTypeRegistry::Get().GetDataTypeInfo(DataTypeName, DataTypeInfo);
+					*bValid &= DataTypeInfo.bIsFloatParsable;
+				}
+				break;
+
+				case EAudioParameterType::FloatArray:
+				{
+					FDataTypeRegistryInfo DataTypeInfo;
+					*bValid = IDataTypeRegistry::Get().GetDataTypeInfo(DataTypeName, DataTypeInfo);
+					*bValid &= DataTypeInfo.bIsFloatArrayParsable;
+				}
+				break;
+
+				case EAudioParameterType::Integer:
+				{
+					FDataTypeRegistryInfo DataTypeInfo;
+					*bValid = IDataTypeRegistry::Get().GetDataTypeInfo(DataTypeName, DataTypeInfo);
+					*bValid &= DataTypeInfo.bIsIntParsable;
+				}
+				break;
+
+				case EAudioParameterType::IntegerArray:
+				{
+					FDataTypeRegistryInfo DataTypeInfo;
+					*bValid = IDataTypeRegistry::Get().GetDataTypeInfo(DataTypeName, DataTypeInfo);
+					*bValid &= DataTypeInfo.bIsIntArrayParsable;
+
+				}
+				break;
+
+				case EAudioParameterType::Object:
+				{
+					FDataTypeRegistryInfo DataTypeInfo;
+					*bValid = IDataTypeRegistry::Get().GetDataTypeInfo(Param->ObjectParam, DataTypeInfo);
+					*bValid &= DataTypeInfo.bIsProxyParsable;
+					*bValid &= DataTypeInfo.DataTypeName == DataTypeName;
+				}
+				break;
+
+				case EAudioParameterType::ObjectArray:
+				{
+					*bValid = true;
+
+					for (UObject* Object : Param->ArrayObjectParam)
+					{
+						FDataTypeRegistryInfo DataTypeInfo;
+						*bValid = IDataTypeRegistry::Get().GetDataTypeInfo(Param->ObjectParam, DataTypeInfo);
+						*bValid &= DataTypeInfo.bIsProxyParsable;
+						*bValid &= DataTypeInfo.DataTypeName == DataTypeName;
+
+						if (!*bValid)
+						{
+							break;
+						}
+					}
+				}
+				break;
+
+				case EAudioParameterType::String:
+				{
+					FDataTypeRegistryInfo DataTypeInfo;
+					*bValid = IDataTypeRegistry::Get().GetDataTypeInfo(DataTypeName, DataTypeInfo);
+					*bValid &= DataTypeInfo.bIsStringParsable;
+				}
+				break;
+
+				case EAudioParameterType::StringArray:
+				{
+					FDataTypeRegistryInfo DataTypeInfo;
+					*bValid = IDataTypeRegistry::Get().GetDataTypeInfo(DataTypeName, DataTypeInfo);
+					*bValid &= DataTypeInfo.bIsStringArrayParsable;
+				}
+				break;
+
+				case EAudioParameterType::NoneArray:
+				{
+					FDataTypeRegistryInfo DataTypeInfo;
+					*bValid = IDataTypeRegistry::Get().GetDataTypeInfo(DataTypeName, DataTypeInfo);
+					*bValid &= DataTypeInfo.bIsDefaultArrayParsable;
+				}
+				case EAudioParameterType::None:
+				default:
+				{
+					FDataTypeRegistryInfo DataTypeInfo;
+					*bValid = IDataTypeRegistry::Get().GetDataTypeInfo(DataTypeName, DataTypeInfo);
+					*bValid &= DataTypeInfo.bIsDefaultParsable;
+				}
+				break;
+			}
+		}
+	}, EMetasoundFrontendClassType::Input);
+
+	return bIsValid;
+}
+
+TUniquePtr<Audio::IParameterTransmitter> UMetaSoundSource::CreateParameterTransmitter(Audio::FParameterTransmitterInitParams&& InParams) const
 {
-	Metasound::FMetasoundInstanceTransmitter::FInitParams InitParams(GetOperatorSettings(InParams.SampleRate), InParams.InstanceID);
+	Metasound::FMetaSoundParameterTransmitter::FInitParams InitParams(GetOperatorSettings(InParams.SampleRate), InParams.InstanceID);
 
 	for (const FSendInfoAndVertexName& InfoAndName : FMetasoundAssetBase::GetSendInfos(InParams.InstanceID))
 	{
 		InitParams.Infos.Add(InfoAndName.SendInfo);
 	}
 
-	return MakeUnique<Metasound::FMetasoundInstanceTransmitter>(InitParams);
+	TUniquePtr<Audio::IParameterTransmitter> NewTransmitter = MakeUnique<Metasound::FMetaSoundParameterTransmitter>(InitParams);
+	for (FAudioParameter& AudioParam : InParams.DefaultParams)
+	{
+		NewTransmitter->SetParameter(MoveTemp(AudioParam));
+	}
+	InParams.DefaultParams.Reset();
+
+	return NewTransmitter;
 }
 
 Metasound::FOperatorSettings UMetaSoundSource::GetOperatorSettings(Metasound::FSampleRate InSampleRate) const
@@ -390,7 +673,7 @@ Metasound::FMetasoundEnvironment UMetaSoundSource::CreateEnvironment(const FSoun
 	return Environment;
 }
 
-Metasound::FMetasoundEnvironment UMetaSoundSource::CreateEnvironment(const FAudioInstanceTransmitterInitParams& InParams) const
+Metasound::FMetasoundEnvironment UMetaSoundSource::CreateEnvironment(const Audio::FParameterTransmitterInitParams& InParams) const
 {
 	using namespace Metasound;
 	using namespace Metasound::Engine;
diff --git a/Engine/Plugins/Runtime/Metasound/Source/MetasoundEngine/Public/MetasoundSettings.h b/Engine/Plugins/Runtime/Metasound/Source/MetasoundEngine/Public/MetasoundSettings.h
--- a/Engine/Plugins/Runtime/Metasound/Source/MetasoundEngine/Public/MetasoundSettings.h
+++ b/Engine/Plugins/Runtime/Metasound/Source/MetasoundEngine/Public/MetasoundSettings.h
@@ -13,7 +13,6 @@
 
 // Forward Declarations
 struct FMetasoundFrontendClassName;
-class IMetaSoundAssetInterface;
 
 
 UENUM()
diff --git a/Engine/Plugins/Runtime/Metasound/Source/MetasoundEngine/Public/MetasoundSource.h b/Engine/Plugins/Runtime/Metasound/Source/MetasoundEngine/Public/MetasoundSource.h
--- a/Engine/Plugins/Runtime/Metasound/Source/MetasoundEngine/Public/MetasoundSource.h
+++ b/Engine/Plugins/Runtime/Metasound/Source/MetasoundEngine/Public/MetasoundSource.h
@@ -3,13 +3,14 @@
 
 #include "CoreMinimal.h"
 #include "EdGraph/EdGraph.h"
+#include "IAudioParameterTransmitter.h"
 #include "Metasound.h"
 #include "MetasoundAssetBase.h"
 #include "MetasoundFrontend.h"
 #include "MetasoundFrontendDocument.h"
 #include "MetasoundFrontendRegistries.h"
-#include "MetasoundInstanceTransmitter.h"
 #include "MetasoundOperatorSettings.h"
+#include "MetasoundParameterTransmitter.h"
 #include "MetasoundRouter.h"
 #include "Sound/SoundWaveProcedural.h"
 #include "UObject/MetaData.h"
@@ -157,13 +158,15 @@ public:
 		return this;
 	}
 
+	virtual void InitParameters(TArray<FAudioParameter>& InParametersToInit, FName InFeatureName) override;
 	virtual void InitResources() override;
 
 	virtual bool IsPlayable() const override;
 	virtual bool SupportsSubtitles() const override;
 	virtual float GetDuration() override;
 	virtual ISoundGeneratorPtr CreateSoundGenerator(const FSoundGeneratorInitParams& InParams) override;
-	virtual TUniquePtr<IAudioInstanceTransmitter> CreateInstanceTransmitter(const FAudioInstanceTransmitterInitParams& InParams) const override;
+	virtual TUniquePtr<Audio::IParameterTransmitter> CreateParameterTransmitter(Audio::FParameterTransmitterInitParams&& InParams) const override;
+	virtual bool IsParameterValid(const FAudioParameter& InParameter) const override;
 
 	// Get the most up to date archetype for metasound sources.
 	const TArray<FMetasoundFrontendVersion>& GetSupportedArchetypeVersions() const override;
@@ -192,6 +195,6 @@ private:
 	Metasound::FOperatorSettings GetOperatorSettings(Metasound::FSampleRate InSampleRate) const;
 	Metasound::FMetasoundEnvironment CreateEnvironment() const;
 	Metasound::FMetasoundEnvironment CreateEnvironment(const FSoundGeneratorInitParams& InParams) const;
-	Metasound::FMetasoundEnvironment CreateEnvironment(const FAudioInstanceTransmitterInitParams& InParams) const;
+	Metasound::FMetasoundEnvironment CreateEnvironment(const Audio::FParameterTransmitterInitParams& InParams) const;
 	const TArray<Metasound::FVertexKey>& GetAudioOutputVertexKeys() const;
 };
diff --git a/Engine/Plugins/Runtime/Metasound/Source/MetasoundFrontend/MetasoundFrontend.Build.cs b/Engine/Plugins/Runtime/Metasound/Source/MetasoundFrontend/MetasoundFrontend.Build.cs
--- a/Engine/Plugins/Runtime/Metasound/Source/MetasoundFrontend/MetasoundFrontend.Build.cs
+++ b/Engine/Plugins/Runtime/Metasound/Source/MetasoundFrontend/MetasoundFrontend.Build.cs
@@ -13,6 +13,7 @@ namespace UnrealBuildTool.Rules
 			(
 				new string[]
 				{
+					"AudioExtensions",
 					"Core",
 					"CoreUObject",
 					"Serialization",
diff --git a/Engine/Plugins/Runtime/Metasound/Source/MetasoundFrontend/Private/MetasoundAssetBase.cpp b/Engine/Plugins/Runtime/Metasound/Source/MetasoundFrontend/Private/MetasoundAssetBase.cpp
--- a/Engine/Plugins/Runtime/Metasound/Source/MetasoundFrontend/Private/MetasoundAssetBase.cpp
+++ b/Engine/Plugins/Runtime/Metasound/Source/MetasoundFrontend/Private/MetasoundAssetBase.cpp
@@ -17,9 +17,9 @@
 #include "MetasoundFrontendRegistries.h"
 #include "MetasoundFrontendSearchEngine.h"
 #include "MetasoundFrontendTransform.h"
-#include "MetasoundInstanceTransmitter.h"
 #include "MetasoundJsonBackend.h"
 #include "MetasoundLog.h"
+#include "MetasoundParameterTransmitter.h"
 #include "MetasoundTrace.h"
 #include "StructSerializer.h"
 #include "UObject/MetaData.h"
@@ -441,7 +441,7 @@ TUniquePtr<Metasound::IGraph> FMetasoundAssetBase::BuildMetasoundDocument() cons
 	if (FrontendGraph.IsValid())
 	{
 		TSet<FVertexKey> VerticesToSkip = GetNonTransmittableInputVertices(*Doc);
-		bool bSuccessfullyInjectedReceiveNodes = InjectReceiveNodes(*FrontendGraph, FMetasoundInstanceTransmitter::CreateSendAddressFromEnvironment, VerticesToSkip);
+		bool bSuccessfullyInjectedReceiveNodes = InjectReceiveNodes(*FrontendGraph, FMetaSoundParameterTransmitter::CreateSendAddressFromEnvironment, VerticesToSkip);
 		if (!bSuccessfullyInjectedReceiveNodes)
 		{
 			UE_LOG(LogMetaSound, Error, TEXT("Error while injecting async communication hooks. Instance communication may not function properly [Name:%s]."), *GetOwningAssetName());
@@ -551,7 +551,7 @@ TArray<FMetasoundAssetBase::FSendInfoAndVertexName> FMetasoundAssetBase::GetSend
 {
 	using namespace Metasound;
 	using namespace Metasound::Frontend;
-	using FSendInfo = FMetasoundInstanceTransmitter::FSendInfo;
+	using FSendInfo = FMetaSoundParameterTransmitter::FSendInfo;
 
 	TArray<FSendInfoAndVertexName> SendInfos;
 
@@ -567,7 +567,7 @@ TArray<FMetasoundAssetBase::FSendInfoAndVertexName> FMetasoundAssetBase::GetSend
 
 			// TODO: incorporate VertexID into address. But need to ensure that VertexID
 			// will be maintained after injecting Receive nodes. 
-			Info.SendInfo.Address = FMetasoundInstanceTransmitter::CreateSendAddressFromInstanceID(InInstanceID, InputHandle->GetName(), InputHandle->GetDataType());
+			Info.SendInfo.Address = FMetaSoundParameterTransmitter::CreateSendAddressFromInstanceID(InInstanceID, InputHandle->GetName(), InputHandle->GetDataType());
 			Info.SendInfo.ParameterName = FName(InputHandle->GetDisplayName().ToString()); // TODO: display name hack. Need to have naming consistent in editor for inputs
 			//Info.SendInfo.ParameterName = FName(*InputHandle->GetName()); // TODO: this is the expected parameter name.
 			Info.SendInfo.TypeName = InputHandle->GetDataType();
@@ -675,7 +675,7 @@ TArray<FString> FMetasoundAssetBase::GetTransmittableInputVertexNames() const
 	return GraphInputVertexNames;
 }
 
-Metasound::Frontend::FNodeHandle FMetasoundAssetBase::AddInputPinForSendAddress(const Metasound::FMetasoundInstanceTransmitter::FSendInfo& InSendInfo, Metasound::Frontend::FGraphHandle InGraph) const
+Metasound::Frontend::FNodeHandle FMetasoundAssetBase::AddInputPinForSendAddress(const Metasound::FMetaSoundParameterTransmitter::FSendInfo& InSendInfo, Metasound::Frontend::FGraphHandle InGraph) const
 {
 	FMetasoundFrontendClassInput Description;
 	FGuid VertexID = FGuid::NewGuid();
diff --git a/Engine/Plugins/Runtime/Metasound/Source/MetasoundFrontend/Private/MetasoundFrontendDataTypeRegistry.cpp b/Engine/Plugins/Runtime/Metasound/Source/MetasoundFrontend/Private/MetasoundFrontendDataTypeRegistry.cpp
--- a/Engine/Plugins/Runtime/Metasound/Source/MetasoundFrontend/Private/MetasoundFrontendDataTypeRegistry.cpp
+++ b/Engine/Plugins/Runtime/Metasound/Source/MetasoundFrontend/Private/MetasoundFrontendDataTypeRegistry.cpp
@@ -326,6 +326,7 @@ namespace Metasound
 
 				virtual void GetRegisteredDataTypeNames(TArray<FName>& OutNames) const override;
 
+				virtual bool GetDataTypeInfo(const UObject* InObject, FDataTypeRegistryInfo& OutInfo) const override;
 				virtual bool GetDataTypeInfo(const FName& InDataType, FDataTypeRegistryInfo& OutInfo) const override;
 
 				// Return the enum interface for a data type. If the data type does not have 
@@ -339,9 +340,12 @@ namespace Metasound
 
 				virtual UClass* GetUClassForDataType(const FName& InDataType) const override;
 
+				bool IsUObjectProxyFactory(UObject* InObject) const override;
+				Audio::IProxyDataPtr CreateProxyFromUObject(const FName& InDataType, UObject* InObject) const override;
+
 				virtual FLiteral CreateDefaultLiteral(const FName& InDataType) const override;
 				virtual FLiteral CreateLiteralFromUObject(const FName& InDataType, UObject* InObject) const override;
-				virtual FLiteral CreateLiteralFromUObjectArray(const FName& InDataType, TArray<UObject*> InObjectArray) const override;
+				virtual FLiteral CreateLiteralFromUObjectArray(const FName& InDataType, const TArray<UObject*>& InObjectArray) const override;
 
 				virtual TSharedPtr<IDataChannel, ESPMode::ThreadSafe> CreateDataChannel(const FName& InDataType, const FOperatorSettings& InOperatorSettings) const override;
 
@@ -362,6 +366,9 @@ namespace Metasound
 				const IDataTypeRegistryEntry* FindDataTypeEntry(const FName& InDataTypeName) const;
 
 				TMap<FName, TUniquePtr<IDataTypeRegistryEntry>> RegisteredDataTypes;
+
+				// UObject type names to DataTypeNames
+				TMap<const UClass*, FName> RegisteredObjectClasses;
 			};
 
 			bool FDataTypeRegistry::RegisterDataType(TUniquePtr<IDataTypeRegistryEntry>&& InEntry)
@@ -387,6 +394,15 @@ namespace Metasound
 						NodeRegistry->RegisterNode(MakeUnique<FVariableNodeRegistryEntry>(InEntry->Clone()));
 					}
 
+					const FDataTypeRegistryInfo& RegistryInfo = InEntry->GetDataTypeInfo();
+					if (!RegistryInfo.IsArrayType())
+					{
+						if (const UClass* Class = RegistryInfo.ProxyGeneratorClass)
+						{
+							RegisteredObjectClasses.Add(Class, Name);
+						}
+					}
+
 					RegisteredDataTypes.Add(Name, MoveTemp(InEntry));
 
 					UE_LOG(LogMetaSound, Verbose, TEXT("Registered Metasound Datatype [Name:%s]."), *Name.ToString());
@@ -401,6 +417,23 @@ namespace Metasound
 				RegisteredDataTypes.GetKeys(OutNames);
 			}
 
+			bool FDataTypeRegistry::GetDataTypeInfo(const UObject* InObject, FDataTypeRegistryInfo& OutInfo) const
+			{
+				if (InObject)
+				{
+					if (const FName* DataTypeName = RegisteredObjectClasses.Find(InObject->GetClass()))
+					{
+						if (const IDataTypeRegistryEntry* Entry = FindDataTypeEntry(*DataTypeName))
+						{
+							OutInfo = Entry->GetDataTypeInfo();
+							return true;
+						}
+					}
+				}
+
+				return false;
+			}
+
 			bool FDataTypeRegistry::GetDataTypeInfo(const FName& InDataType, FDataTypeRegistryInfo& OutInfo) const
 			{
 				if (const IDataTypeRegistryEntry* Entry = FindDataTypeEntry(InDataType))
@@ -529,55 +562,56 @@ namespace Metasound
 				return nullptr;
 			}
 
-			FLiteral FDataTypeRegistry::CreateDefaultLiteral(const FName& InDataType) const 
+			FLiteral FDataTypeRegistry::CreateDefaultLiteral(const FName& InDataType) const
 			{
 				return FLiteral::GetDefaultForType(GetDesiredLiteralType(InDataType));
 			}
 
-			FLiteral FDataTypeRegistry::CreateLiteralFromUObject(const FName& InDataType, UObject* InObject) const
+			bool FDataTypeRegistry::IsUObjectProxyFactory(UObject* InObject) const
 			{
-				if (const IDataTypeRegistryEntry* Entry = FindDataTypeEntry(InDataType))
+				if (!InObject)
 				{
-					if (Audio::IProxyDataPtr ProxyPtr = Entry->CreateProxy(InObject))
-					{
-						if (InObject)
-						{
-							ensureAlwaysMsgf(ProxyPtr.IsValid(), TEXT("UObject failed to create a valid proxy."));
-						}
-						return Metasound::FLiteral(MoveTemp(ProxyPtr));
-					}
+					return false;
 				}
 
-				return Metasound::FLiteral();
+				return RegisteredObjectClasses.Contains(InObject->GetClass());
 			}
 
-			FLiteral FDataTypeRegistry::CreateLiteralFromUObjectArray(const FName& InDataType, TArray<UObject*> InObjectArray) const
+			Audio::IProxyDataPtr FDataTypeRegistry::CreateProxyFromUObject(const FName& InDataType, UObject* InObject) const
 			{
+				Audio::IProxyDataPtr ProxyPtr;
+
 				if (const IDataTypeRegistryEntry* Entry = FindDataTypeEntry(InDataType))
 				{
-					TArray<Audio::IProxyDataPtr> ProxyArray;
-
-					for (UObject* InObject : InObjectArray)
+					ProxyPtr = Entry->CreateProxy(InObject);
+					if (ProxyPtr)
 					{
-						Audio::IProxyDataPtr ProxyPtr = Entry->CreateProxy(InObject);
-
-						if (InObject)
+						if (!InObject)
 						{
-							if (!ProxyPtr.IsValid())
-							{
-								UE_LOG(LogMetaSound, Error, TEXT("Failed to create a valid proxy from UObject [Name:%s]."), *InObject->GetName());
-							}
+							UE_LOG(LogMetaSound, Error, TEXT("Failed to create a valid proxy from UObject '%s'."), *InObject->GetName());
 						}
-
-						ProxyArray.Add(MoveTemp(ProxyPtr));
 					}
-
-					return Metasound::FLiteral(MoveTemp(ProxyArray));
 				}
-				else
+
+				return ProxyPtr;
+			}
+
+			FLiteral FDataTypeRegistry::CreateLiteralFromUObject(const FName& InDataType, UObject* InObject) const
+			{
+				Audio::IProxyDataPtr ProxyPtr = CreateProxyFromUObject(InDataType, InObject);
+				return Metasound::FLiteral(MoveTemp(ProxyPtr));
+			}
+
+			FLiteral FDataTypeRegistry::CreateLiteralFromUObjectArray(const FName& InDataType, const TArray<UObject*>& InObjectArray) const
+			{
+				TArray<Audio::IProxyDataPtr> ProxyArray;
+				for (UObject* InObject : InObjectArray)
 				{
-					return Metasound::FLiteral();
+					Audio::IProxyDataPtr ProxyPtr = CreateProxyFromUObject(InDataType, InObject);
+					ProxyArray.Emplace(MoveTemp(ProxyPtr));
 				}
+
+				return Metasound::FLiteral(MoveTemp(ProxyArray));
 			}
 
 			TSharedPtr<IDataChannel, ESPMode::ThreadSafe> FDataTypeRegistry::CreateDataChannel(const FName& InDataType, const FOperatorSettings& InOperatorSettings) const
diff --git a/Engine/Plugins/Runtime/Metasound/Source/MetasoundFrontend/Private/MetasoundInstanceTransmitter.cpp b/Engine/Plugins/Runtime/Metasound/Source/MetasoundFrontend/Private/MetasoundInstanceTransmitter.cpp
deleted file mode 100644
diff --git a/Engine/Plugins/Runtime/Metasound/Source/MetasoundFrontend/Private/MetasoundParameterTransmitter.cpp b/Engine/Plugins/Runtime/Metasound/Source/MetasoundFrontend/Private/MetasoundParameterTransmitter.cpp
new file mode 100644
--- /dev/null
+++ b/Engine/Plugins/Runtime/Metasound/Source/MetasoundFrontend/Private/MetasoundParameterTransmitter.cpp
@@ -0,0 +1,196 @@
+// Copyright Epic Games, Inc. All Rights Reserved.
+
+#include "MetasoundParameterTransmitter.h"
+
+#include "MetasoundLog.h"
+
+
+namespace Metasound
+{
+	FString FMetaSoundParameterTransmitter::GetInstanceIDEnvironmentVariableName()
+	{
+		return TEXT("TransmitterInstanceID");
+	}
+
+	FSendAddress FMetaSoundParameterTransmitter::CreateSendAddressFromEnvironment(const FMetasoundEnvironment& InEnvironment, const FVertexKey& InVertexKey, const FName& InTypeName)
+	{
+		const FString IDVarName = GetInstanceIDEnvironmentVariableName();
+		uint64 InstanceID = -1;
+
+		if (ensure(InEnvironment.Contains<uint64>(IDVarName)))
+		{
+			InstanceID = InEnvironment.GetValue<uint64>(IDVarName);
+		}
+
+		return CreateSendAddressFromInstanceID(InstanceID, InVertexKey, InTypeName);
+	}
+
+	FSendAddress FMetaSoundParameterTransmitter::CreateSendAddressFromInstanceID(uint64 InInstanceID, const FVertexKey& InVertexKey, const FName& InTypeName)
+	{
+		FSendAddress Address;
+
+		Address.Subsystem = GetSubsystemNameForSendScope(ETransmissionScope::Global);
+		Address.ChannelName = FName(FString::Format(TEXT("{0}:{1}:{2}"), { InInstanceID, InVertexKey, InTypeName.ToString() }));
+
+		return Address;
+	}
+
+	FMetaSoundParameterTransmitter::FMetaSoundParameterTransmitter(const FMetaSoundParameterTransmitter::FInitParams& InInitParams)
+	: SendInfos(InInitParams.Infos)
+	, OperatorSettings(InInitParams.OperatorSettings)
+	, InstanceID(InInitParams.InstanceID)
+	{
+	}
+
+	bool FMetaSoundParameterTransmitter::Reset()
+	{
+		bool bSuccess = true;
+
+		for (const FSendInfo& SendInfo : SendInfos)
+		{
+			bSuccess &= FDataTransmissionCenter::Get().UnregisterDataChannel(SendInfo.TypeName, SendInfo.Address);
+		}
+
+		return bSuccess;
+	}
+
+	uint64 FMetaSoundParameterTransmitter::GetInstanceID() const
+	{
+		return InstanceID;
+	}
+
+	bool FMetaSoundParameterTransmitter::SetParameter(FAudioParameter&& InValue)
+	{
+		switch (InValue.ParamType)
+		{
+			case EAudioParameterType::Boolean:
+			{
+				return SetParameterWithLiteral(InValue.ParamName, FLiteral(InValue.BoolParam));
+			}
+
+			case EAudioParameterType::BooleanArray:
+			{
+				return SetParameterWithLiteral(InValue.ParamName, FLiteral(MoveTemp(InValue.ArrayBoolParam)));
+			}
+
+			case EAudioParameterType::Float:
+			{
+				return SetParameterWithLiteral(InValue.ParamName, FLiteral(InValue.FloatParam));
+			}
+
+			case EAudioParameterType::FloatArray:
+			{
+				return SetParameterWithLiteral(InValue.ParamName, FLiteral(MoveTemp(InValue.ArrayFloatParam)));
+			}
+
+			case EAudioParameterType::Integer:
+			{
+				return SetParameterWithLiteral(InValue.ParamName, FLiteral(InValue.IntParam));
+			}
+
+			case EAudioParameterType::IntegerArray:
+			{
+				return SetParameterWithLiteral(InValue.ParamName, FLiteral(MoveTemp(InValue.ArrayIntParam)));
+			}
+
+			case EAudioParameterType::None:
+			{
+				return SetParameterWithLiteral(InValue.ParamName, FLiteral());
+			}
+
+			case EAudioParameterType::NoneArray:
+			{
+				TArray<FLiteral::FNone> InitArray;
+				InitArray.Init(FLiteral::FNone(), InValue.IntParam);
+				return SetParameterWithLiteral(InValue.ParamName, FLiteral(MoveTemp(InitArray)));
+			}
+
+			case EAudioParameterType::Object:
+			{
+				if(InValue.ObjectProxies.IsEmpty())
+				{
+					return SetParameterWithLiteral(InValue.ParamName, FLiteral());
+				}
+
+				return SetParameterWithLiteral(InValue.ParamName, FLiteral(MoveTemp(InValue.ObjectProxies[0])));
+			}
+
+			case EAudioParameterType::ObjectArray:
+			{
+				return SetParameterWithLiteral(InValue.ParamName, FLiteral(MoveTemp(InValue.ObjectProxies)));
+			}
+
+			case EAudioParameterType::String:
+			{
+				return SetParameterWithLiteral(InValue.ParamName, FLiteral(MoveTemp(InValue.StringParam)));
+			}
+
+			case EAudioParameterType::StringArray:
+			{
+				return SetParameterWithLiteral(InValue.ParamName, FLiteral(MoveTemp(InValue.ArrayStringParam)));
+			}
+
+			default:
+			{
+				checkNoEntry();
+			}
+		}
+
+		return false;
+	}
+
+	bool FMetaSoundParameterTransmitter::SetParameterWithLiteral(FName InParameterName, const FLiteral& InLiteral)
+	{
+		if (ISender* Sender = FindSender(InParameterName))
+		{
+			return Sender->PushLiteral(InLiteral);
+		}
+
+		// If no sender exists for parameter name, attempt to add one.
+		if (const FSendInfo* SendInfo = FindSendInfo(InParameterName))
+		{
+			if (ISender* Sender = AddSender(*SendInfo))
+			{
+				return Sender->PushLiteral(InLiteral);
+			}
+		}
+
+		return false;
+	}
+
+	TUniquePtr<Audio::IParameterTransmitter> FMetaSoundParameterTransmitter::Clone() const
+	{
+		return MakeUnique<FMetaSoundParameterTransmitter>(FMetaSoundParameterTransmitter::FInitParams(OperatorSettings, InstanceID, SendInfos));
+	}
+
+	const FMetaSoundParameterTransmitter::FSendInfo* FMetaSoundParameterTransmitter::FindSendInfo(const FName& InParameterName) const
+	{
+		return SendInfos.FindByPredicate([&](const FSendInfo& Info) { return Info.ParameterName == InParameterName; });
+	}
+
+	ISender* FMetaSoundParameterTransmitter::FindSender(const FName& InParameterName)
+	{
+		if (TUniquePtr<ISender>* SenderPtrPtr = InputSends.Find(InParameterName))
+		{
+			return SenderPtrPtr->Get();
+		}
+		return nullptr;
+	}
+
+	ISender* FMetaSoundParameterTransmitter::AddSender(const FSendInfo& InInfo)
+	{
+		// TODO: likely want to remove this and opt for different protocols having different behaviors.
+		const float DelayTimeInSeconds = 0.1f; // This not used for non-audio routing.
+		const FSenderInitParams InitParams = { OperatorSettings, DelayTimeInSeconds };
+
+		TUniquePtr<ISender> Sender = FDataTransmissionCenter::Get().RegisterNewSender(InInfo.TypeName, InInfo.Address, InitParams);
+		if (ensureMsgf(Sender.IsValid(), TEXT("Failed to create sender [DataType:%s, Address:%s]"), *InInfo.TypeName.ToString(), *InInfo.Address.ChannelName.ToString()))
+		{
+			ISender* Ptr = Sender.Get();
+			InputSends.Add(InInfo.ParameterName, MoveTemp(Sender));
+			return Ptr;
+		}
+
+		return nullptr;
+	}
+}
diff --git a/Engine/Plugins/Runtime/Metasound/Source/MetasoundFrontend/Public/MetasoundAssetBase.h b/Engine/Plugins/Runtime/Metasound/Source/MetasoundFrontend/Public/MetasoundAssetBase.h
--- a/Engine/Plugins/Runtime/Metasound/Source/MetasoundFrontend/Public/MetasoundAssetBase.h
+++ b/Engine/Plugins/Runtime/Metasound/Source/MetasoundFrontend/Public/MetasoundAssetBase.h
@@ -9,8 +9,8 @@
 #include "MetasoundFrontendDocument.h"
 #include "MetasoundFrontendDocumentAccessPtr.h"
 #include "MetasoundGraph.h"
-#include "MetasoundInstanceTransmitter.h"
 #include "MetasoundLog.h"
+#include "MetasoundParameterTransmitter.h"
 #include "UObject/SoftObjectPath.h"
 #include "UObject/WeakObjectPtrTemplates.h"
 
@@ -195,7 +195,7 @@ public:
 
 	struct FSendInfoAndVertexName
 	{
-		Metasound::FMetasoundInstanceTransmitter::FSendInfo SendInfo;
+		Metasound::FMetaSoundParameterTransmitter::FSendInfo SendInfo;
 		FString VertexName;
 	};
 
@@ -232,5 +232,5 @@ private:
 	TSet<Metasound::FVertexKey> GetNonTransmittableInputVertices(const FMetasoundFrontendDocument& InDoc) const;
 	TArray<FString> GetTransmittableInputVertexNames() const;
 	Metasound::FSendAddress CreateSendAddress(uint64 InInstanceID, const FString& InVertexName, const FName& InDataTypeName) const;
-	Metasound::Frontend::FNodeHandle AddInputPinForSendAddress(const Metasound::FMetasoundInstanceTransmitter::FSendInfo& InSendInfo, Metasound::Frontend::FGraphHandle InGraph) const;
+	Metasound::Frontend::FNodeHandle AddInputPinForSendAddress(const Metasound::FMetaSoundParameterTransmitter::FSendInfo& InSendInfo, Metasound::Frontend::FGraphHandle InGraph) const;
 };
diff --git a/Engine/Plugins/Runtime/Metasound/Source/MetasoundFrontend/Public/MetasoundDataTypeRegistrationMacro.h b/Engine/Plugins/Runtime/Metasound/Source/MetasoundFrontend/Public/MetasoundDataTypeRegistrationMacro.h
--- a/Engine/Plugins/Runtime/Metasound/Source/MetasoundFrontend/Public/MetasoundDataTypeRegistrationMacro.h
+++ b/Engine/Plugins/Runtime/Metasound/Source/MetasoundFrontend/Public/MetasoundDataTypeRegistrationMacro.h
@@ -5,6 +5,7 @@
 #include "CoreMinimal.h"
 #include "Templates/Casts.h"
 
+#include "IAudioProxyInitializer.h"
 #include "MetasoundArrayNodesRegistration.h"
 #include "MetasoundAutoConverterNode.h"
 #include "MetasoundConverterNodeRegistrationMacro.h"
@@ -147,29 +148,6 @@ namespace Metasound
 			static constexpr ELiteralType Value = ELiteralType::UObjectProxyArray;
 		};
 
-		// SFINAE used to optionally invoke subclasses of IAudioProxyDataFactory when we can.
-		template<typename UClassToUse, typename TEnableIf<TIsDerivedFrom<UClassToUse, IAudioProxyDataFactory>::Value, bool>::Type = true>
-		IAudioProxyDataFactory* CastToAudioProxyDataFactory(UObject* InObject)
-		{
-			if (InObject)
-			{
-				UClassToUse* DowncastObject = Cast<UClassToUse>(InObject);
-				if (ensureAlways(DowncastObject))
-				{
-					return static_cast<IAudioProxyDataFactory*>(DowncastObject);
-				}
-			}
-
-			return nullptr;
-		}
-
-		template<typename UClassToUse, typename TEnableIf<!TIsDerivedFrom<UClassToUse, IAudioProxyDataFactory>::Value, bool>::Type = true>
-		IAudioProxyDataFactory* CastToAudioProxyDataFactory(UObject* InObject)
-		{
-			return nullptr;
-		}
-		
-
 		// This utility function can be used to optionally check to see if we can transmit a data type, and autogenerate send and receive nodes for that datatype.
 		template<typename TDataType, typename TEnableIf<TIsTransmittable<TDataType>::Value, bool>::Type = true>
 		void AttemptToRegisterSendAndReceiveNodes()
@@ -434,15 +412,13 @@ namespace Metasound
 					if constexpr (!std::is_same<UClassToUse, void>::value)
 					{
 						static_assert(std::is_base_of<IAudioProxyDataFactory, UClassToUse>::value, "If a Metasound Datatype uses a UObject as a literal, the UClass of that object needs to also derive from Audio::IProxyDataFactory. See USoundWave as an example.");
-						if (InObject)
+						if (Frontend::IDataTypeRegistry::Get().IsUObjectProxyFactory(InObject))
 						{
-							IAudioProxyDataFactory* ObjectAsFactory = CastToAudioProxyDataFactory<UClassToUse>(InObject);
+							IAudioProxyDataFactory* ObjectAsFactory = Audio::CastToProxyDataFactory<UClassToUse>(InObject);
 							if (ensureAlways(ObjectAsFactory))
 							{
-								static FName ProxySubsystemName = TEXT("Metasound");
-
 								Audio::FProxyDataInitParams ProxyInitParams;
-								ProxyInitParams.NameOfFeatureRequestingProxy = ProxySubsystemName;
+								ProxyInitParams.NameOfFeatureRequestingProxy = "MetaSound";
 
 								return ObjectAsFactory->CreateNewProxyData(ProxyInitParams);
 							}
diff --git a/Engine/Plugins/Runtime/Metasound/Source/MetasoundFrontend/Public/MetasoundFrontendDataTypeRegistry.h b/Engine/Plugins/Runtime/Metasound/Source/MetasoundFrontend/Public/MetasoundFrontendDataTypeRegistry.h
--- a/Engine/Plugins/Runtime/Metasound/Source/MetasoundFrontend/Public/MetasoundFrontendDataTypeRegistry.h
+++ b/Engine/Plugins/Runtime/Metasound/Source/MetasoundFrontend/Public/MetasoundFrontendDataTypeRegistry.h
@@ -146,7 +146,7 @@ namespace Metasound
 
 			/** Create a proxy from a UObject. If this data type does not support
 			 * UObject proxies, return a nullptr. */
-			virtual Audio::IProxyDataPtr CreateProxy(UObject*) const = 0;
+			virtual Audio::IProxyDataPtr CreateProxy(UObject* InObject) const = 0;
 
 			/** Create a data channel for transmission. If this data type does not
 			 * support transmission, return a nullptr. */
@@ -172,6 +172,7 @@ namespace Metasound
 
 			virtual void GetRegisteredDataTypeNames(TArray<FName>& OutNames) const = 0;
 
+			virtual bool GetDataTypeInfo(const UObject* InObject, FDataTypeRegistryInfo& OutInfo) const = 0;
 			virtual bool GetDataTypeInfo(const FName& InDataType, FDataTypeRegistryInfo& OutInfo) const = 0;
 
 			// Return the enum interface for a data type. If the data type does not have 
@@ -182,12 +183,15 @@ namespace Metasound
 
 			virtual bool IsLiteralTypeSupported(const FName& InDataType, ELiteralType InLiteralType) const = 0;
 			virtual bool IsLiteralTypeSupported(const FName& InDataType, EMetasoundFrontendLiteralType InLiteralType) const = 0;
+			virtual bool IsUObjectProxyFactory(UObject* InObject) const = 0;
 
 			virtual UClass* GetUClassForDataType(const FName& InDataType) const = 0;
 
+			virtual Audio::IProxyDataPtr CreateProxyFromUObject(const FName& InDataType, UObject* InObject) const = 0;
+
 			virtual FLiteral CreateDefaultLiteral(const FName& InDataType) const = 0;
 			virtual FLiteral CreateLiteralFromUObject(const FName& InDataType, UObject* InObject) const = 0;
-			virtual FLiteral CreateLiteralFromUObjectArray(const FName& InDataType, TArray<UObject*> InObjectArray) const = 0;
+			virtual FLiteral CreateLiteralFromUObjectArray(const FName& InDataType, const TArray<UObject*>& InObjectArray) const = 0;
 
 			virtual TSharedPtr<IDataChannel, ESPMode::ThreadSafe> CreateDataChannel(const FName& InDataType, const FOperatorSettings& InOperatorSettings) const = 0;
 
diff --git a/Engine/Plugins/Runtime/Metasound/Source/MetasoundFrontend/Public/MetasoundFrontendLiteral.h b/Engine/Plugins/Runtime/Metasound/Source/MetasoundFrontend/Public/MetasoundFrontendLiteral.h
--- a/Engine/Plugins/Runtime/Metasound/Source/MetasoundFrontend/Public/MetasoundFrontendLiteral.h
+++ b/Engine/Plugins/Runtime/Metasound/Source/MetasoundFrontend/Public/MetasoundFrontendLiteral.h
@@ -2,6 +2,7 @@
 
 #pragma once
 
+#include "AudioParameterInterface.h"
 #include "CoreMinimal.h"
 #include "MetasoundLiteral.h"
 
@@ -32,10 +33,23 @@ enum class EMetasoundFrontendLiteralType : uint8
 	StringArray = static_cast<uint8>(Metasound::ELiteralType::StringArray),
 	UObjectArray = static_cast<uint8>(Metasound::ELiteralType::UObjectProxyArray),
 
-
 	Invalid UMETA(Hidden)
 };
 
+static_assert(static_cast<uint8>(EMetasoundFrontendLiteralType::None) == static_cast<uint8>(EAudioParameterType::None), "Type 'None' value must match");
+static_assert(static_cast<uint8>(EMetasoundFrontendLiteralType::Boolean) == static_cast<uint8>(EAudioParameterType::Boolean), "Type 'Boolean' value must match");
+static_assert(static_cast<uint8>(EMetasoundFrontendLiteralType::Integer) == static_cast<uint8>(EAudioParameterType::Integer), "Type 'Integer' value must match");
+static_assert(static_cast<uint8>(EMetasoundFrontendLiteralType::Float) == static_cast<uint8>(EAudioParameterType::Float), "Type 'Float' value must match");
+static_assert(static_cast<uint8>(EMetasoundFrontendLiteralType::String) == static_cast<uint8>(EAudioParameterType::String), "Type 'String' value must match");
+static_assert(static_cast<uint8>(EMetasoundFrontendLiteralType::UObject) == static_cast<uint8>(EAudioParameterType::Object), "Type 'UObjectProxy' value must match");
+static_assert(static_cast<uint8>(EMetasoundFrontendLiteralType::NoneArray) == static_cast<uint8>(EAudioParameterType::NoneArray), "Type 'NoneArray' value must match");
+static_assert(static_cast<uint8>(EMetasoundFrontendLiteralType::BooleanArray) == static_cast<uint8>(EAudioParameterType::BooleanArray), "Type 'BooleanArray' value must match");
+static_assert(static_cast<uint8>(EMetasoundFrontendLiteralType::IntegerArray) == static_cast<uint8>(EAudioParameterType::IntegerArray), "Type 'IntegerArray' value must match");
+static_assert(static_cast<uint8>(EMetasoundFrontendLiteralType::FloatArray) == static_cast<uint8>(EAudioParameterType::FloatArray), "Type 'FloatArray' value must match");
+static_assert(static_cast<uint8>(EMetasoundFrontendLiteralType::StringArray) == static_cast<uint8>(EAudioParameterType::StringArray), "Type 'StringArray' value must match");
+static_assert(static_cast<uint8>(EMetasoundFrontendLiteralType::UObjectArray) == static_cast<uint8>(EAudioParameterType::ObjectArray), "Type 'UObjectProxyArray' value must match");
+static_assert(static_cast<uint8>(EMetasoundFrontendLiteralType::Invalid) == static_cast<uint8>(EAudioParameterType::COUNT), "Enum EMetasoundFrontendLiteralType' count must match 'EAudioParameterType'");
+
 // Check that the static_cast<>s above are using the correct type.
 static_assert(std::is_same<uint8, std::underlying_type<EMetasoundFrontendLiteralType>::type>::value, "Update type in static_cast<TYPE> from Metasound::ELiteralType to EMetasoundFrontendLiteralType in EMetasoundFrontendLiteralType declaration.");
 static_assert(std::is_same<std::underlying_type<Metasound::ELiteralType>::type, std::underlying_type<EMetasoundFrontendLiteralType>::type>::value, "EMetasoundFrontendLiteralType and Metasound::ELiteralType must have matching underlying types to support conversion.");
diff --git a/Engine/Plugins/Runtime/Metasound/Source/MetasoundFrontend/Public/MetasoundInstanceTransmitter.h b/Engine/Plugins/Runtime/Metasound/Source/MetasoundFrontend/Public/MetasoundInstanceTransmitter.h
deleted file mode 100644
diff --git a/Engine/Plugins/Runtime/Metasound/Source/MetasoundFrontend/Public/MetasoundParameterTransmitter.h b/Engine/Plugins/Runtime/Metasound/Source/MetasoundFrontend/Public/MetasoundParameterTransmitter.h
new file mode 100644
--- /dev/null
+++ b/Engine/Plugins/Runtime/Metasound/Source/MetasoundFrontend/Public/MetasoundParameterTransmitter.h
@@ -0,0 +1,117 @@
+// Copyright Epic Games, Inc. All Rights Reserved.
+#pragma once
+
+#include "AudioParameterInterface.h"
+#include "IAudioParameterTransmitter.h"
+#include "MetasoundDataReference.h"
+#include "MetasoundLog.h"
+#include "MetasoundOperatorSettings.h"
+#include "MetasoundPrimitives.h"
+#include "MetasoundRouter.h"
+
+
+namespace Metasound
+{
+	/** FMetaSoundParameterTransmitter provides a communication interface for 
+	 * sending values to a MetaSound instance. It relies on the send/receive transmission
+	 * system to ferry data from the transmitter to the MetaSound instance. Data will
+	 * be safely ushered across thread boundaries in scenarios where the instance
+	 * transmitter and metasound instance live on different threads. 
+	 */
+	class METASOUNDFRONTEND_API FMetaSoundParameterTransmitter : public Audio::IParameterTransmitter
+	{
+		FMetaSoundParameterTransmitter(const FMetaSoundParameterTransmitter&) = delete;
+		FMetaSoundParameterTransmitter& operator=(const FMetaSoundParameterTransmitter&) = delete;
+	public:
+		
+		/** FSendInfo describes the MetaSounds input parameters as well as the 
+		 * necessary information to route data to the instances inputs. 
+		 */
+		struct FSendInfo
+		{
+			/** Global address of instance input. */
+			FSendAddress Address;
+
+			/** Name of parameter on MetaSound instance. */
+			FName ParameterName;
+
+			/** Type name of parameter on MetaSound instance. */
+			FName TypeName;
+		};
+
+
+		/** Initialization parameters for a FMetaSoundParameterTransmitter. */
+		struct FInitParams
+		{
+			/** FOperatorSettings must match the operator settings of the MetaSound 
+			 * instance to ensure proper operation. */
+			FOperatorSettings OperatorSettings;
+
+			/** ID of the MetaSound instance.  */
+			uint64 InstanceID;
+
+			/** Available input parameters on MetaSound instance. */
+			TArray<FSendInfo> Infos;
+
+			FInitParams(const FOperatorSettings& InSettings, uint64 InInstanceID, const TArray<FSendInfo>& InInfos=TArray<FSendInfo>())
+			: OperatorSettings(InSettings)
+			, InstanceID(InInstanceID)
+			, Infos(InInfos)
+			{
+			}
+
+		};
+
+		/** Returns the MetaSound environment variable name which contains the instance ID. */
+		static FString GetInstanceIDEnvironmentVariableName();
+
+		/** Creates a unique send address using the given MetaSound environment. */
+		static FSendAddress CreateSendAddressFromEnvironment(const FMetasoundEnvironment& InEnvironment, const FVertexKey& InVertexkey, const FName& InTypeName);
+		
+		/** Creates a unique send address using the given InstanceID. */
+		static FSendAddress CreateSendAddressFromInstanceID(uint64 InInstanceID, const FVertexKey& InVertexKey, const FName& InTypeName);
+
+		FMetaSoundParameterTransmitter(const FMetaSoundParameterTransmitter::FInitParams& InInitParams);
+		virtual ~FMetaSoundParameterTransmitter() = default;
+
+		bool Reset() override;
+
+		/** Returns ID of the MetaSound instance associated with this transmitter. */
+		uint64 GetInstanceID() const override;
+
+		/** Sets a parameter using an AudioParameter struct
+		 *
+		 * @param InParameter - Parameter to set.
+		 */
+		bool SetParameter(FAudioParameter&& InParameter) override;
+
+		/** Set a parameter using a literal.
+		 *
+		 * @param InParameterName - Name of MetaSound instance parameter.
+		 * @param InValue - Literal value used to construct parameter value. 
+		 *
+		 * @return true on success, false on failure. 
+		 */
+		bool SetParameterWithLiteral(FName InParameterName, const FLiteral& InValue);
+
+		/** Duplicate this transmitter interface. The transmitters association with
+		 * the MetaSound instance will be maintained. */
+		TUniquePtr<Audio::IParameterTransmitter> Clone() const override;
+
+	private:
+		// Find FSendInfo by parameter name. 
+		const FSendInfo* FindSendInfo(const FName& InParameterName) const;
+
+		// Find ISender by parameter name. 
+		ISender* FindSender(const FName& InParameterName);
+
+		// Create and store a new ISender for the given FSendInfo.
+		ISender* AddSender(const FSendInfo& InInfo);
+
+		TArray<FSendInfo> SendInfos;
+		FOperatorSettings OperatorSettings;
+		uint64 InstanceID;
+
+		TMap<FName, TUniquePtr<ISender>> InputSends;
+	};
+}
diff --git a/Engine/Plugins/Runtime/Metasound/Source/MetasoundFrontend/Public/MetasoundRouter.h b/Engine/Plugins/Runtime/Metasound/Source/MetasoundFrontend/Public/MetasoundRouter.h
--- a/Engine/Plugins/Runtime/Metasound/Source/MetasoundFrontend/Public/MetasoundRouter.h
+++ b/Engine/Plugins/Runtime/Metasound/Source/MetasoundFrontend/Public/MetasoundRouter.h
@@ -10,6 +10,7 @@
 #include "MetasoundDataReference.h"
 #include "MetasoundExecutableOperator.h"
 #include "MetasoundLiteral.h"
+#include "MetasoundLog.h"
 #include "MetasoundOperatorInterface.h"
 #include "MetasoundOperatorSettings.h"
 #include "MetasoundTrigger.h"
@@ -43,6 +44,7 @@
  * 4) while the overall system is thread safe, individual TSender and TReceiver objects should only be used by a single thread.
  */
 
+
 namespace Metasound
 {
 	// Enum for which scope to use 
@@ -358,6 +360,8 @@ namespace Metasound
 		// This can be changed to change when we reallocate the internal buffer for this sender.
 		static constexpr float BufferResizeThreshold = 2.0f;
 
+		bool bMessageDataTypeError = true;
+
 	public:
 
 		bool Push(const TDataType& InElement)
@@ -378,9 +382,11 @@ namespace Metasound
 		{
 			if (!TLiteralTraits<TDataType>::IsParsable(InLiteral))
 			{
-				// The literal sent in is not compatible with the data type.
-				ensureMsgf(false, TEXT("Failed to send data. Data type [TypeName:%s] cannot be parsed from literal [Literal:%s]."), *GetMetasoundDataTypeString<TDataType>(), *LexToString(InLiteral));
-
+				if (bMessageDataTypeError)
+				{
+					UE_LOG(LogMetaSound, Error, TEXT("Failed to send data. Data type [TypeName:%s] cannot be parsed from literal [Literal:%s]."), *GetMetasoundDataTypeString<TDataType>(), *LexToString(InLiteral));
+					bMessageDataTypeError = false;
+				}
 				return false;
 			}
 			else
diff --git a/Engine/Source/Runtime/AudioExtensions/Private/AudioParameterInterface.cpp b/Engine/Source/Runtime/AudioExtensions/Private/AudioParameterInterface.cpp
new file mode 100644
--- /dev/null
+++ b/Engine/Source/Runtime/AudioExtensions/Private/AudioParameterInterface.cpp
@@ -0,0 +1,164 @@
+// Copyright Epic Games, Inc. All Rights Reserved.
+
+#include "AudioParameterInterface.h"
+#include "IAudioProxyInitializer.h"
+
+
+namespace AudioParameterPrivate
+{
+	template <typename T>
+	void SetOrMergeArray(const TArray<T>& InArray, TArray<T>& OutArray, bool bInMerge)
+	{
+		if (bInMerge)
+		{
+			OutArray.Append(InArray);
+		}
+		else
+		{
+			OutArray = InArray;
+		}
+	}
+}
+
+void FAudioParameter::Merge(const FAudioParameter& InParameter, bool bInTakeName, bool bInTakeType, bool bInMergeArrayTypes)
+{
+	if (bInTakeName)
+	{
+		ParamName = InParameter.ParamName;
+	}
+
+	if (bInTakeType)
+	{
+		ParamType = InParameter.ParamType;
+	}
+
+	switch (InParameter.ParamType)
+	{
+		case EAudioParameterType::Boolean:
+		{
+			BoolParam = InParameter.BoolParam;
+		}
+		break;
+
+		case EAudioParameterType::BooleanArray:
+		{
+			AudioParameterPrivate::SetOrMergeArray(InParameter.ArrayBoolParam, ArrayBoolParam, bInMergeArrayTypes);
+		}
+		break;
+
+		case EAudioParameterType::Float:
+		{
+			FloatParam = InParameter.FloatParam;
+		}
+		break;
+
+		case EAudioParameterType::FloatArray:
+		{
+			AudioParameterPrivate::SetOrMergeArray(InParameter.ArrayFloatParam, ArrayFloatParam, bInMergeArrayTypes);
+		}
+		break;
+
+		case EAudioParameterType::Integer:
+		case EAudioParameterType::NoneArray:
+		{
+			if (bInMergeArrayTypes)
+			{
+				IntParam += InParameter.IntParam;
+			}
+			else
+			{
+				IntParam = InParameter.IntParam;
+			}
+		}
+		break;
+
+		case EAudioParameterType::IntegerArray:
+		{
+			AudioParameterPrivate::SetOrMergeArray(InParameter.ArrayIntParam, ArrayIntParam, bInMergeArrayTypes);
+		}
+		break;
+
+		case EAudioParameterType::None:
+		{
+			FloatParam = InParameter.FloatParam;
+			BoolParam = InParameter.BoolParam;
+			IntParam = InParameter.IntParam;
+			ObjectParam = InParameter.ObjectParam;
+			StringParam = InParameter.StringParam;
+
+			AudioParameterPrivate::SetOrMergeArray(InParameter.ArrayBoolParam, ArrayBoolParam, bInMergeArrayTypes);
+			AudioParameterPrivate::SetOrMergeArray(InParameter.ArrayFloatParam, ArrayFloatParam, bInMergeArrayTypes);
+			AudioParameterPrivate::SetOrMergeArray(InParameter.ArrayIntParam, ArrayIntParam, bInMergeArrayTypes);
+			AudioParameterPrivate::SetOrMergeArray(InParameter.ArrayObjectParam, ArrayObjectParam, bInMergeArrayTypes);
+			AudioParameterPrivate::SetOrMergeArray(InParameter.ArrayStringParam, ArrayStringParam, bInMergeArrayTypes);
+
+			if (!bInMergeArrayTypes)
+			{
+				ObjectProxies.Reset();
+			}
+
+			for (const Audio::IProxyDataPtr& ProxyPtr : InParameter.ObjectProxies)
+			{
+				if (ensure(ProxyPtr.IsValid()))
+				{
+					ObjectProxies.Emplace(ProxyPtr->Clone());
+				}
+			}
+		}
+		break;
+
+		case EAudioParameterType::Object:
+		{
+			ObjectParam = InParameter.ObjectParam;
+
+			ObjectProxies.Reset();
+			for (const Audio::IProxyDataPtr& ProxyPtr : InParameter.ObjectProxies)
+			{
+				if (ensure(ProxyPtr.IsValid()))
+				{
+					ObjectProxies.Emplace(ProxyPtr->Clone());
+				}
+			}
+		}
+		break;
+
+		case EAudioParameterType::ObjectArray:
+		{
+			AudioParameterPrivate::SetOrMergeArray(InParameter.ArrayObjectParam, ArrayObjectParam, bInMergeArrayTypes);
+
+			if (!bInMergeArrayTypes)
+			{
+				ObjectProxies.Reset();
+			}
+
+			for (const Audio::IProxyDataPtr& ProxyPtr : InParameter.ObjectProxies)
+			{
+				if (ensure(ProxyPtr.IsValid()))
+				{
+					ObjectProxies.Emplace(ProxyPtr->Clone());
+				}
+			}
+		}
+		break;
+
+		case EAudioParameterType::String:
+		{
+			StringParam = InParameter.StringParam;
+		}
+		break;
+
+		case EAudioParameterType::StringArray:
+		{
+			AudioParameterPrivate::SetOrMergeArray(InParameter.ArrayStringParam, ArrayStringParam, bInMergeArrayTypes);
+		}
+		break;
+
+		default:
+			break;
+	}
+}
+
+UAudioParameterInterface::UAudioParameterInterface(FObjectInitializer const& InObjectInitializer)
+	: UInterface(InObjectInitializer)
+{
+}
diff --git a/Engine/Source/Runtime/AudioExtensions/Public/AudioParameterInterface.h b/Engine/Source/Runtime/AudioExtensions/Public/AudioParameterInterface.h
new file mode 100644
--- /dev/null
+++ b/Engine/Source/Runtime/AudioExtensions/Public/AudioParameterInterface.h
@@ -0,0 +1,433 @@
+// Copyright Epic Games, Inc. All Rights Reserved.
+
+#pragma once
+
+#include "CoreMinimal.h"
+#include "IAudioProxyInitializer.h"
+#include "UObject/Interface.h"
+
+#include "AudioParameterInterface.generated.h"
+
+
+UENUM()
+enum class EAudioParameterType : uint8
+{
+	// 'Default' results in behavior that is resolved
+	// based on the system interpreting it.  To support
+	// legacy implementation, SoundCues cache all typed values
+	// associated with a given parameter name, whereas MetaSounds
+	// only use the last typed value set.
+	None UMETA(DisplayName = "Default"),
+
+	// Boolean value
+	Boolean,
+
+	// Integer value
+	Integer,
+
+	// Float value
+	Float,
+
+	// String value (not supported by legacy SoundCue system)
+	String,
+
+	// Object value (types other than SoundWave not supported by legacy SoundCue system)
+	Object,
+
+	// Array of default initialized values (not supported by legacy SoundCue system)
+	// Hidden for now as no parameter types exist that support default construction
+	NoneArray UMETA(Hidden, DisplayName = "Default (Array)"),
+
+	// Array of boolean values (not supported by legacy SoundCue system)
+	BooleanArray UMETA(DisplayName = "Boolean (Array)"),
+
+	// Array of integer values (not supported by legacy SoundCue system)
+	IntegerArray UMETA(DisplayName = "Integer (Array)"),
+
+	// Array of float values (not supported by legacy SoundCue system)
+	FloatArray UMETA(DisplayName = "Float (Array)"),
+
+	// Array of string values (not supported by legacy SoundCue system)
+	StringArray UMETA(DisplayName = "String (Array)"),
+
+	// Array of object values (not supported by legacy SoundCue system)
+	ObjectArray UMETA(DisplayName = "Object (Array)"),
+
+	COUNT UMETA(Hidden)
+};
+
+
+USTRUCT(BlueprintType)
+struct AUDIOEXTENSIONS_API FAudioParameter
+{
+	GENERATED_USTRUCT_BODY()
+
+	FAudioParameter() = default;
+
+	FAudioParameter(FName InName)
+		: ParamName(InName)
+	{
+	}
+
+	FAudioParameter(FName InName, float InValue)
+		: ParamName(InName)
+		, FloatParam(InValue)
+		, ParamType(EAudioParameterType::Float)
+	{
+	}
+
+	FAudioParameter(FName InName, bool InValue)
+		: ParamName(InName)
+		, BoolParam(InValue)
+		, ParamType(EAudioParameterType::Boolean)
+	{
+	}
+
+	FAudioParameter(FName InName, int32 InValue)
+		: ParamName(InName)
+		, IntParam(InValue)
+		, ParamType(EAudioParameterType::Integer)
+	{
+	}
+
+	FAudioParameter(FName InName, UObject* InValue)
+		: ParamName(InName)
+		, ObjectParam(InValue)
+		, ParamType(EAudioParameterType::Object)
+	{
+	}
+
+	FAudioParameter(FName InName, const FString& InValue)
+		: ParamName(InName)
+		, StringParam(InValue)
+		, ParamType(EAudioParameterType::String)
+	{
+	}
+
+	FAudioParameter(FName InName, const TArray<float>& InValue)
+		: ParamName(InName)
+		, ArrayFloatParam(InValue)
+		, ParamType(EAudioParameterType::FloatArray)
+	{
+	}
+
+	FAudioParameter(FName InName, TArray<float>&& InValue)
+		: ParamName(InName)
+		, ArrayFloatParam(MoveTemp(InValue))
+		, ParamType(EAudioParameterType::FloatArray)
+	{
+	}
+
+	FAudioParameter(FName InName, const TArray<bool>& InValue)
+		: ParamName(InName)
+		, ArrayBoolParam(InValue)
+		, ParamType(EAudioParameterType::BooleanArray)
+	{
+	}
+
+	FAudioParameter(FName InName, TArray<bool>&& InValue)
+		: ParamName(InName)
+		, ArrayBoolParam(MoveTemp(InValue))
+		, ParamType(EAudioParameterType::BooleanArray)
+	{
+	}
+
+	FAudioParameter(FName InName, const TArray<int32>& InValue)
+		: ParamName(InName)
+		, ArrayIntParam(InValue)
+		, ParamType(EAudioParameterType::IntegerArray)
+	{
+	}
+
+	FAudioParameter(FName InName, TArray<int32>&& InValue)
+		: ParamName(InName)
+		, ArrayIntParam(MoveTemp(InValue))
+		, ParamType(EAudioParameterType::IntegerArray)
+	{
+	}
+
+	FAudioParameter(FName InName, const TArray<UObject*>& InValue)
+		: ParamName(InName)
+		, ArrayObjectParam(InValue)
+		, ParamType(EAudioParameterType::ObjectArray)
+	{
+	}
+
+	FAudioParameter(FName InName, TArray<UObject*>&& InValue)
+		: ParamName(InName)
+		, ArrayObjectParam(MoveTemp(InValue))
+		, ParamType(EAudioParameterType::ObjectArray)
+	{
+	}
+
+	FAudioParameter(FName InName, const TArray<FString>& InValue)
+		: ParamName(InName)
+		, ArrayStringParam(InValue)
+		, ParamType(EAudioParameterType::StringArray)
+	{
+	}
+
+	FAudioParameter(FName InName, TArray<FString>&& InValue)
+		: ParamName(InName)
+		, ArrayStringParam(MoveTemp(InValue))
+		, ParamType(EAudioParameterType::StringArray)
+	{
+	}
+
+	// Static function to avoid int32 constructor collision
+	static FAudioParameter CreateDefaultArray(FName InName, int32 InNum)
+	{
+		FAudioParameter NewParam(InName, InNum);
+		NewParam.ParamType = EAudioParameterType::NoneArray;
+		return NewParam;
+	}
+
+	// Sets values specified by type field of the given parameter on this parameter. If the provided parameter is set to type 'None', takes all values of the given parameter.
+	// bInTakeName = Take the name of the provided parameter.
+	// bInTakeType = Take the type of the provided parameter.
+	// bInMergeArrayTypes - Appends array(s) for specified type if true, else swaps the local value with that of the provided parameter if false.
+	void Merge(const FAudioParameter& InParameter, bool bInTakeName = true, bool bInTakeType = true, bool bInMergeArrayTypes = false);
+
+	FAudioParameter(FAudioParameter&& InParameter)
+		: ParamName(MoveTemp(InParameter.ParamName))
+		, FloatParam(MoveTemp(InParameter.FloatParam))
+		, BoolParam(MoveTemp(InParameter.BoolParam))
+		, IntParam(MoveTemp(InParameter.IntParam))
+		, ObjectParam(MoveTemp(InParameter.ObjectParam))
+		, StringParam(MoveTemp(InParameter.StringParam))
+		, ArrayFloatParam(MoveTemp(InParameter.ArrayFloatParam))
+		, ArrayBoolParam(MoveTemp(InParameter.ArrayBoolParam))
+		, ArrayIntParam(MoveTemp(InParameter.ArrayIntParam))
+		, ArrayObjectParam(MoveTemp(InParameter.ArrayObjectParam))
+		, ArrayStringParam(MoveTemp(InParameter.ArrayStringParam))
+		, ParamType(MoveTemp(InParameter.ParamType))
+		, ObjectProxies(MoveTemp(InParameter.ObjectProxies))
+	{
+	}
+
+	FAudioParameter(const FAudioParameter& InParameter)
+		: ParamName(InParameter.ParamName)
+		, FloatParam(InParameter.FloatParam)
+		, BoolParam(InParameter.BoolParam)
+		, IntParam(InParameter.IntParam)
+		, ObjectParam(InParameter.ObjectParam)
+		, StringParam(InParameter.StringParam)
+		, ArrayFloatParam(InParameter.ArrayFloatParam)
+		, ArrayBoolParam(InParameter.ArrayBoolParam)
+		, ArrayIntParam(InParameter.ArrayIntParam)
+		, ArrayObjectParam(InParameter.ArrayObjectParam)
+		, ArrayStringParam(InParameter.ArrayStringParam)
+		, ParamType(InParameter.ParamType)
+	{
+		ObjectProxies.Reset();
+		for (const Audio::IProxyDataPtr& DataPtr : InParameter.ObjectProxies)
+		{
+			ObjectProxies.Emplace(DataPtr->Clone());
+		}
+	}
+
+	FAudioParameter& operator=(const FAudioParameter& InParameter)
+	{
+		ParamName = InParameter.ParamName;
+		FloatParam = InParameter.FloatParam;
+		BoolParam = InParameter.BoolParam;
+		IntParam = InParameter.IntParam;
+		ObjectParam = InParameter.ObjectParam;
+		StringParam = InParameter.StringParam;
+		ArrayFloatParam = InParameter.ArrayFloatParam;
+		ArrayBoolParam = InParameter.ArrayBoolParam;
+		ArrayIntParam = InParameter.ArrayIntParam;
+		ArrayObjectParam = InParameter.ArrayObjectParam;
+		ArrayStringParam = InParameter.ArrayStringParam;
+		ParamType = InParameter.ParamType;
+
+		ObjectProxies.Reset();
+		for (const Audio::IProxyDataPtr& DataPtr : InParameter.ObjectProxies)
+		{
+			ObjectProxies.Emplace(DataPtr->Clone());
+		}
+
+		return *this;
+	}
+
+	FAudioParameter& operator=(FAudioParameter&& InParameter)
+	{
+		ParamName = MoveTemp(InParameter.ParamName);
+		FloatParam = MoveTemp(InParameter.FloatParam);
+		BoolParam = MoveTemp(InParameter.BoolParam);
+		IntParam = MoveTemp(InParameter.IntParam);
+		ObjectParam = MoveTemp(InParameter.ObjectParam);
+		StringParam = MoveTemp(InParameter.StringParam);
+		ArrayFloatParam = MoveTemp(InParameter.ArrayFloatParam);
+		ArrayBoolParam = MoveTemp(InParameter.ArrayBoolParam);
+		ArrayIntParam = MoveTemp(InParameter.ArrayIntParam);
+		ArrayObjectParam = MoveTemp(InParameter.ArrayObjectParam);
+		ArrayStringParam = MoveTemp(InParameter.ArrayStringParam);
+		ParamType = MoveTemp(InParameter.ParamType);
+
+		ObjectProxies = MoveTemp(InParameter.ObjectProxies);
+
+		return *this;
+	}
+
+	// Name of the parameter
+	UPROPERTY(EditAnywhere, BlueprintReadWrite, meta = (DisplayName="Name"), Category = AudioParameter)
+	FName ParamName;
+
+	// Float value of parameter
+	UPROPERTY(EditAnywhere, BlueprintReadWrite, meta = (DisplayName = "Value (Float)", DisplayAfter = "ParamType", EditConditionHides, EditCondition = "ParamType == EAudioParameterType::None || ParamType == EAudioParameterType::Float"), Category = AudioParameter)
+	float FloatParam = 0.f;
+
+	// Boolean value of parameter
+	UPROPERTY(EditAnywhere, BlueprintReadWrite, meta = (DisplayName = "Value (Bool)", DisplayAfter = "ParamType", EditConditionHides, EditCondition = "ParamType == EAudioParameterType::None || ParamType == EAudioParameterType::Bool"), Category = AudioParameter)
+	bool BoolParam = false;
+
+	// Integer value of parameter. If set to 'Default Construct', value is number of array items to construct.
+	UPROPERTY(EditAnywhere, BlueprintReadWrite, meta = (DisplayName = "Value (Int)", DisplayAfter = "ParamType", EditConditionHides, EditCondition = "ParamType == EAudioParameterType::None || ParamType == EAudioParameterType::Integer || ParamType == EAudioParameterType::NoneArray"), Category = AudioParameter)
+	int32 IntParam = 0;
+
+	// Object value of parameter
+	UPROPERTY(EditAnywhere, BlueprintReadWrite, meta = (DisplayName = "Value (Object)", DisplayAfter = "ParamType", EditConditionHides, EditCondition = "ParamType == EAudioParameterType::None || ParamType == EAudioParameterType::Object"), Category = AudioParameter)
+	TObjectPtr<UObject> ObjectParam = nullptr;
+
+	// String value of parameter
+	UPROPERTY(EditAnywhere, BlueprintReadWrite, meta = (DisplayName = "Value", DisplayAfter = "ParamType", EditConditionHides, EditCondition = "ParamType == EAudioParameterType::String"), Category = AudioParameter)
+	FString StringParam;
+
+	// Array Float value of parameter
+	UPROPERTY(EditAnywhere, BlueprintReadWrite, meta = (DisplayName = "Value", DisplayAfter = "ParamType", EditConditionHides, EditCondition = "ParamType == EAudioParameterType::FloatArray"), Category = AudioParameter)
+	TArray<float> ArrayFloatParam;
+
+	// Boolean value of parameter
+	UPROPERTY(EditAnywhere, BlueprintReadWrite, meta = (DisplayName = "Value", DisplayAfter = "ParamType", EditConditionHides, EditCondition = "ParamType == EAudioParameterType::BooleanArray"), Category = AudioParameter)
+	TArray<bool> ArrayBoolParam;
+
+	// Integer value of parameter
+	UPROPERTY(EditAnywhere, BlueprintReadWrite, meta = (DisplayName = "Value", DisplayAfter = "ParamType", EditConditionHides, EditCondition = "ParamType == EAudioParameterType::IntegerArray"), Category = AudioParameter)
+	TArray<int32> ArrayIntParam;
+
+	// Object value of parameter
+	UPROPERTY(EditAnywhere, BlueprintReadWrite, meta = (DisplayName = "Value", DisplayAfter = "ParamType", EditConditionHides, EditCondition = "ParamType == EAudioParameterType::ObjectArray"), Category = AudioParameter)
+	TArray<TObjectPtr<UObject>> ArrayObjectParam;
+
+	// String value of parameter
+	UPROPERTY(EditAnywhere, BlueprintReadWrite, meta = (DisplayName = "Value", DisplayAfter = "ParamType", EditConditionHides, EditCondition = "ParamType == EAudioParameterType::StringArray"), Category = AudioParameter)
+	TArray<FString> ArrayStringParam;
+
+	UPROPERTY(EditAnywhere, BlueprintReadWrite, meta = (DisplayName = "Type"), Category = AudioParameter)
+	EAudioParameterType ParamType = EAudioParameterType::None;
+
+	// Object proxies to be generated when parameter is passed to the AudioThread to represent ObjectParam/ArrayObjectParam safely
+	TArray<Audio::IProxyDataPtr> ObjectProxies;
+
+	// Common find algorithm for default/legacy parameter system
+	static const FAudioParameter* FindParam(const TArray<FAudioParameter>& InParams, FName InParamName)
+	{
+		if (InParamName.IsValid())
+		{
+			for (const FAudioParameter& ExistingParam : InParams)
+			{
+				if (ExistingParam.ParamName == InParamName)
+				{
+					return &ExistingParam;
+				}
+			}
+		}
+
+		return nullptr;
+	}
+
+	// Common find & add algorithm for default/legacy parameter system.
+	static FAudioParameter* FindOrAddParam(TArray<FAudioParameter>& OutParams, FName InParamName)
+	{
+		FAudioParameter* Param = nullptr;
+		if (!InParamName.IsValid())
+		{
+			return Param;
+		}
+
+		for (FAudioParameter& ExistingParam : OutParams)
+		{
+			if (ExistingParam.ParamName == InParamName)
+			{
+				Param = &ExistingParam;
+				break;
+			}
+		}
+
+		if (!Param)
+		{
+			Param = &OutParams.AddDefaulted_GetRef();
+			Param->ParamName = InParamName;
+		}
+
+		return Param;
+	}
+};
+
+
+UINTERFACE(BlueprintType, meta = (CannotImplementInterfaceInBlueprint))
+class AUDIOEXTENSIONS_API UAudioParameterInterface : public UInterface
+{
+	GENERATED_UINTERFACE_BODY()
+};
+
+class AUDIOEXTENSIONS_API IAudioParameterInterface : public IInterface
+{
+	GENERATED_IINTERFACE_BODY()
+
+public:
+	// Resets all parameters to their original values.
+	UFUNCTION(BlueprintCallable, Category = "Audio|Parameter")
+	virtual void ResetParameters() = 0;
+
+	// Triggers a named trigger
+	UFUNCTION(BlueprintCallable, meta = (DisplayName = "Set Trigger Parameter"), Category = "Audio|Parameter")
+	virtual void SetTriggerParameter(FName InName) = 0;
+
+	// Sets a named Boolean
+	UFUNCTION(BlueprintCallable, meta = (DisplayName = "Set Boolean Parameter"), Category = "Audio|Parameter")
+	virtual void SetBoolParameter(FName InName, bool InBool) = 0;
+	
+	// Sets a named Boolean Array
+	UFUNCTION(BlueprintCallable, meta = (DisplayName = "Set Boolean Array Parameter"), Category = "Audio|Parameter")
+	virtual void SetBoolArrayParameter(FName InName, const TArray<bool>& InValue) = 0;
+
+	// Sets a named Int32
+	UFUNCTION(BlueprintCallable, meta = (DisplayName = "Set Integer Parameter"), Category = "Audio|Parameter")
+	virtual void SetIntParameter(FName InName, int32 InInt) = 0;
+
+	// Sets a named Int32 Array
+	UFUNCTION(BlueprintCallable, meta = (DisplayName = "Set Integer Array Parameter"), Category = "Audio|Parameter")
+	virtual void SetIntArrayParameter(FName InName, const TArray<int32>& InValue) = 0;
+
+	// Sets a named Float
+	UFUNCTION(BlueprintCallable, meta = (DisplayName = "Set Float Parameter"), Category = "Audio|Parameter")
+	virtual void SetFloatParameter(FName InName, float InFloat) = 0;
+	
+	// Sets a named Float Array
+	UFUNCTION(BlueprintCallable, meta = (DisplayName = "Set Float Array Parameter"), Category = "Audio|Parameter")
+	virtual void SetFloatArrayParameter(FName InName, const TArray<float>& InValue) = 0;
+	
+	// Sets a named String
+	UFUNCTION(BlueprintCallable, meta = (DisplayName = "Set String Parameter"), Category = "Audio|Parameter")
+	virtual void SetStringParameter(FName InName, const FString& InValue) = 0;
+
+	// Sets a named String Array
+	UFUNCTION(BlueprintCallable, meta = (DisplayName = "Set String Array Parameter"), Category = "Audio|Parameter")
+	virtual void SetStringArrayParameter(FName InName, const TArray<FString>& InValue) = 0;
+
+	// Sets a named UObject
+	UFUNCTION(BlueprintCallable, meta = (DisplayName = "Set Object Parameter"), Category = "Audio|Parameter")
+	virtual void SetObjectParameter(FName InName, UObject* InValue) = 0;
+
+	// Sets a named UObject Array
+	UFUNCTION(BlueprintCallable, meta = (DisplayName = "Set Object Array Parameter"), Category = "Audio|Parameter")
+	virtual void SetObjectArrayParameter(FName InName, const TArray<UObject*>& InValue) = 0;
+
+	// Sets a named parameter to the given parameter structure value
+	virtual void SetParameter(FAudioParameter&& InValue) = 0;
+
+	// Sets an array of parameters as a batch
+	virtual void SetParameters(TArray<FAudioParameter>&& InValues) = 0;
+};
\ No newline at end of file
diff --git a/Engine/Source/Runtime/AudioExtensions/Public/IAudioExtensionPlugin.h b/Engine/Source/Runtime/AudioExtensions/Public/IAudioExtensionPlugin.h
--- a/Engine/Source/Runtime/AudioExtensions/Public/IAudioExtensionPlugin.h
+++ b/Engine/Source/Runtime/AudioExtensions/Public/IAudioExtensionPlugin.h
@@ -607,42 +607,3 @@ public:
 	{
 	}
 };
-
-/** Parameters passed to USoundBase::CreateInstanceTransmitter. */
-struct FAudioInstanceTransmitterInitParams
-{
-	// Unique ID for this audio instance.
-	uint64 InstanceID;
-	// Audio sample rate.
-	float SampleRate;
-};
-
-/** Interface for a audio instance transmitter.
- *
- * An audio instance transmitter ushers control parameters to a single audio object instance.
- */
-class IAudioInstanceTransmitter
-{
-	public:
-		virtual ~IAudioInstanceTransmitter() = default;
-
-		virtual bool Shutdown() = 0;
-
-		// Return the instance ID
-		virtual uint64 GetInstanceID() const = 0;
-	
-		// Set Parameter Overloads
-		virtual bool SetParameter(const FName& InName, bool InValue) = 0;
-		virtual bool SetParameter(const FName& InName, TArray<bool>&& InValue) = 0;
-		virtual bool SetParameter(const FName& InName, int32 InValue) = 0;
-		virtual bool SetParameter(const FName& InName, TArray<int32>&& InValue) = 0;
-		virtual bool SetParameter(const FName& InName, float InValue) = 0;
-		virtual bool SetParameter(const FName& InName, TArray<float>&& InValue) = 0;
-		virtual bool SetParameter(const FName& InName, FString&& InValue) = 0;
-		virtual bool SetParameter(const FName& InName, TArray<FString>&& InValue) = 0;
-		virtual bool SetParameter(const FName& InName, Audio::IProxyDataPtr&&) = 0;
-		virtual bool SetParameter(const FName& InName, TArray<Audio::IProxyDataPtr>&& InValue) = 0;
-	
-		// Create a copy of the instance transmitter.
-		virtual TUniquePtr<IAudioInstanceTransmitter> Clone() const = 0;
-};
diff --git a/Engine/Source/Runtime/AudioExtensions/Public/IAudioProxyInitializer.h b/Engine/Source/Runtime/AudioExtensions/Public/IAudioProxyInitializer.h
--- a/Engine/Source/Runtime/AudioExtensions/Public/IAudioProxyInitializer.h
+++ b/Engine/Source/Runtime/AudioExtensions/Public/IAudioProxyInitializer.h
@@ -3,6 +3,7 @@
 #pragma once
 
 #include "CoreMinimal.h"
+#include "Templates/Casts.h"
 #include "Templates/UnrealTypeTraits.h"
 
 /**
@@ -126,13 +127,38 @@ namespace Audio
 	{
 		FName NameOfFeatureRequestingProxy;
 	};
-}
+} // namespace Audio
 
 /*
 * This can be subclassed to make a UClass an audio proxy factory.
 */
-class IAudioProxyDataFactory
+class AUDIOEXTENSIONS_API IAudioProxyDataFactory
 {
 public:
 	virtual TUniquePtr<Audio::IProxyData> CreateNewProxyData(const Audio::FProxyDataInitParams& InitParams) = 0;
 };
+
+namespace Audio
+{
+	// SFINAE used to optionally invoke subclasses of IAudioProxyDataFactory when we can.
+	template<typename UClassToUse, typename TEnableIf<TIsDerivedFrom<UClassToUse, IAudioProxyDataFactory>::Value, bool>::Type = true>
+	IAudioProxyDataFactory* CastToProxyDataFactory(UObject* InObject)
+	{
+		if (InObject)
+		{
+			UClassToUse* DowncastObject = Cast<UClassToUse>(InObject);
+			if (ensureAlways(DowncastObject))
+			{
+				return static_cast<IAudioProxyDataFactory*>(DowncastObject);
+			}
+		}
+
+		return nullptr;
+	}
+
+	template<typename UClassToUse, typename TEnableIf<!TIsDerivedFrom<UClassToUse, IAudioProxyDataFactory>::Value, bool>::Type = true>
+	IAudioProxyDataFactory* CastToProxyDataFactory(UObject* InObject)
+	{
+		return nullptr;
+	}
+} // namespace Audio
diff --git a/Engine/Source/Runtime/Engine/Classes/Audio/AudioComponentParameterization.h b/Engine/Source/Runtime/Engine/Classes/Audio/AudioComponentParameterization.h
deleted file mode 100644
diff --git a/Engine/Source/Runtime/Engine/Classes/Audio/AudioParameterInterface.h b/Engine/Source/Runtime/Engine/Classes/Audio/AudioParameterInterface.h
deleted file mode 100644
diff --git a/Engine/Source/Runtime/Engine/Classes/Audio/SoundGeneratorParameterInterface.h b/Engine/Source/Runtime/Engine/Classes/Audio/SoundGeneratorParameterInterface.h
new file mode 100644
--- /dev/null
+++ b/Engine/Source/Runtime/Engine/Classes/Audio/SoundGeneratorParameterInterface.h
@@ -0,0 +1,59 @@
+// Copyright Epic Games, Inc. All Rights Reserved.
+#pragma once
+
+#include "UObject/Object.h"
+#include "AudioParameterInterface.h"
+
+#include "SoundGeneratorParameterInterface.generated.h"
+
+
+// Forward Declarations
+class FAudioDevice;
+class USoundBase;
+
+UINTERFACE()
+class ENGINE_API USoundGeneratorParameterInterface : public UAudioParameterInterface
+{
+	GENERATED_UINTERFACE_BODY()
+};
+
+class ENGINE_API ISoundGeneratorParameterInterface : public IAudioParameterInterface
+{
+	GENERATED_IINTERFACE_BODY()
+
+public:
+	// IAudioParameterInterface
+	void ResetParameters() override;
+
+	void SetTriggerParameter(FName InName) override;
+	void SetBoolParameter(FName InName, bool InBool) override;
+	void SetBoolArrayParameter(FName InName, const TArray<bool>& InValue) override;
+	void SetIntParameter(FName InName, int32 InInt) override;
+	void SetIntArrayParameter(FName InName, const TArray<int32>& InValue) override;
+	void SetFloatParameter(FName InName, float InFloat) override;
+	void SetFloatArrayParameter(FName InName, const TArray<float>& InValue) override;
+	void SetStringParameter(FName InName, const FString& InValue) override;
+	void SetStringArrayParameter(FName InName, const TArray<FString>& InValue) override;
+	void SetObjectParameter(FName InName, UObject* InValue) override;
+	void SetObjectArrayParameter(FName InName, const TArray<UObject*>& InValue) override;
+
+	void SetParameter(FAudioParameter&& InValue) override;
+	void SetParameters(TArray<FAudioParameter>&& InValues) override;
+
+	/** Returns the active audio device to use for this component based on whether or not the component is playing in a world. */
+	virtual FAudioDevice* GetAudioDevice() const = 0;
+
+	/** Returns the id of the sound instance associated with the parameter interface. */
+	virtual uint64 GetInstanceID() const = 0;
+
+	/** Returns the Game Thread copy of parameters to modify in place. */
+	virtual TArray<FAudioParameter>& GetInstanceParameters() = 0;
+
+	/** Returns the USoundBase used to initialize instance parameters to update. */
+	virtual USoundBase* GetSound() = 0;
+
+	virtual bool IsPlaying() const = 0;
+
+private:
+	void SetParameterInternal(FAudioParameter&& InValue);
+};
diff --git a/Engine/Source/Runtime/Engine/Classes/Components/AudioComponent.h b/Engine/Source/Runtime/Engine/Classes/Components/AudioComponent.h
--- a/Engine/Source/Runtime/Engine/Classes/Components/AudioComponent.h
+++ b/Engine/Source/Runtime/Engine/Classes/Components/AudioComponent.h
@@ -1,26 +1,34 @@
 // Copyright Epic Games, Inc. All Rights Reserved.
 #pragma once
 
-#include "Audio/AudioParameterInterface.h"
+#include "Audio/SoundGeneratorParameterInterface.h"
+#include "AudioParameterInterface.h"
 #include "Components/SceneComponent.h"
 #include "CoreMinimal.h"
 #include "Engine/EngineTypes.h"
-#include "IAudioExtensionPlugin.h"
-#include "Sound/SoundAttenuation.h"
-#include "Sound/SoundWave.h"
-#include "UObject/ObjectMacros.h"
+#include "IAudioParameterTransmitter.h"
 #include "Math/RandomStream.h"
-#include "Sound/QuartzSubscription.h"
-#include "Sound/QuartzQuantizationUtilities.h"
 #include "Quartz/AudioMixerClockHandle.h"
 #include "Quartz/AudioMixerQuantizedCommands.h"
+#include "Sound/QuartzQuantizationUtilities.h"
+#include "Sound/QuartzSubscription.h"
+#include "Sound/SoundAttenuation.h"
+#include "Sound/SoundWave.h"
+#include "UObject/ObjectMacros.h"
 
 #include "AudioComponent.generated.h"
 
+
+// Forward Declarations
 class FAudioDevice;
+class UAudioComponent;
 class USoundBase;
 class USoundClass;
 class USoundConcurrency;
+class USoundWave;
+
+struct FAudioComponentParam;
+
 
 // Enum describing the audio component play state
 UENUM(BlueprintType)
@@ -49,7 +57,7 @@ enum class EAudioComponentPlayState : uint8
 DECLARE_DYNAMIC_MULTICAST_DELEGATE(FOnAudioFinished);
 
 /** shadow delegate declaration for above */
-DECLARE_MULTICAST_DELEGATE_OneParam(FOnAudioFinishedNative, class UAudioComponent*);
+DECLARE_MULTICAST_DELEGATE_OneParam(FOnAudioFinishedNative, UAudioComponent*);
 
 /** Called when subtitles are sent to the SubtitleManager.  Set this delegate if you want to hijack the subtitles for other purposes */
 DECLARE_DYNAMIC_DELEGATE_TwoParams(FOnQueueSubtitles, const TArray<struct FSubtitleCue>&, Subtitles, float, CueDuration);
@@ -58,22 +66,22 @@ DECLARE_DYNAMIC_DELEGATE_TwoParams(FOnQueueSubtitles, const TArray<struct FSubti
 DECLARE_DYNAMIC_MULTICAST_DELEGATE_OneParam(FOnAudioPlayStateChanged, EAudioComponentPlayState, PlayState);
 
 /** shadow delegate declaration for above */
-DECLARE_MULTICAST_DELEGATE_TwoParams(FOnAudioPlayStateChangedNative, const class UAudioComponent*, EAudioComponentPlayState);
+DECLARE_MULTICAST_DELEGATE_TwoParams(FOnAudioPlayStateChangedNative, const UAudioComponent*, EAudioComponentPlayState);
 
 /** Called when sound becomes virtualized or realized (resumes playback from virtualization). */
 DECLARE_DYNAMIC_MULTICAST_DELEGATE_OneParam(FOnAudioVirtualizationChanged, bool, bIsVirtualized);
 
 /** shadow delegate declaration for above */
-DECLARE_MULTICAST_DELEGATE_TwoParams(FOnAudioVirtualizationChangedNative, const class UAudioComponent*, bool);
+DECLARE_MULTICAST_DELEGATE_TwoParams(FOnAudioVirtualizationChangedNative, const UAudioComponent*, bool);
 
 /** Called as a sound plays on the audio component to allow BP to perform actions based on playback percentage.
 * Computed as samples played divided by total samples, taking into account pitch.
 * Not currently implemented on all platforms.
 */
-DECLARE_DYNAMIC_MULTICAST_DELEGATE_TwoParams(FOnAudioPlaybackPercent, const class USoundWave*, PlayingSoundWave, const float, PlaybackPercent);
+DECLARE_DYNAMIC_MULTICAST_DELEGATE_TwoParams(FOnAudioPlaybackPercent, const USoundWave*, PlayingSoundWave, const float, PlaybackPercent);
 
 /** shadow delegate declaration for above */
-DECLARE_MULTICAST_DELEGATE_ThreeParams(FOnAudioPlaybackPercentNative, const class UAudioComponent*, const class USoundWave*, const float);
+DECLARE_MULTICAST_DELEGATE_ThreeParams(FOnAudioPlaybackPercentNative, const UAudioComponent*, const USoundWave*, const float);
 
 /**
 * Called while a sound plays and returns the sound's envelope value (using an envelope follower in the audio renderer).
@@ -82,7 +90,7 @@ DECLARE_MULTICAST_DELEGATE_ThreeParams(FOnAudioPlaybackPercentNative, const clas
 DECLARE_DYNAMIC_MULTICAST_DELEGATE_TwoParams(FOnAudioSingleEnvelopeValue, const class USoundWave*, PlayingSoundWave, const float, EnvelopeValue);
 
 /** shadow delegate declaration for above */
-DECLARE_MULTICAST_DELEGATE_ThreeParams(FOnAudioSingleEnvelopeValueNative, const class UAudioComponent*, const class USoundWave*, const float);
+DECLARE_MULTICAST_DELEGATE_ThreeParams(FOnAudioSingleEnvelopeValueNative, const UAudioComponent*, const USoundWave*, const float);
 
 /**
 * Called while a sound plays and returns the sound's average and max envelope value (using an envelope follower in the audio renderer per wave instance).
@@ -91,7 +99,7 @@ DECLARE_MULTICAST_DELEGATE_ThreeParams(FOnAudioSingleEnvelopeValueNative, const
 DECLARE_DYNAMIC_MULTICAST_DELEGATE_ThreeParams(FOnAudioMultiEnvelopeValue, const float, AverageEnvelopeValue, const float, MaxEnvelope, const int32, NumWaveInstances);
 
 /** shadow delegate declaration for above */
-DECLARE_MULTICAST_DELEGATE_FourParams(FOnAudioMultiEnvelopeValueNative, const class UAudioComponent*, const float, const float, const int32);
+DECLARE_MULTICAST_DELEGATE_FourParams(FOnAudioMultiEnvelopeValueNative, const UAudioComponent*, const float, const float, const int32);
 
 
 
@@ -116,49 +124,16 @@ enum class EAudioFaderCurve : uint8
 
 
 /**
- *	Struct used for storing one per-instance named parameter for this AudioComponent.
- *	Certain nodes in the SoundCue may reference parameters by name so they can be adjusted per-instance.
+ * Legacy struct used for storing named parameter for a given AudioComponent.
  */
 USTRUCT(BlueprintType)
-struct FAudioComponentParam
+struct UE_DEPRECATED(5.0, "FAudioComponentParam has been deprecated, use FAudioParameter") FAudioComponentParam : public FAudioParameter
 {
-	GENERATED_USTRUCT_BODY()
-
-	// Name of the parameter
-	UPROPERTY(EditAnywhere, BlueprintReadWrite, Category=AudioComponentParam)
-	FName ParamName;
-
-	// Value of the parameter when used as a float
-	UPROPERTY(EditAnywhere, BlueprintReadWrite, Category=AudioComponentParam)
-	float FloatParam;
-
-	// Value of the parameter when used as a boolean
-	UPROPERTY(EditAnywhere, BlueprintReadWrite, Category=AudioComponentParam)
-	bool BoolParam;
+	GENERATED_BODY()
 
-	// Value of the parameter when used as an integer
+	// DEPRECATED
 	UPROPERTY(EditAnywhere, BlueprintReadWrite, Category=AudioComponentParam)
-	int32 IntParam;
-
-	// Value of the parameter when used as a sound wave
-	UPROPERTY(EditAnywhere, BlueprintReadWrite, Category=AudioComponentParam)
-	TObjectPtr<class USoundWave> SoundWaveParam;
-
-	FAudioComponentParam(const FName& Name)
-		: ParamName(Name)
-		, FloatParam(0.f)
-		, BoolParam(false)
-		, IntParam(0)
-		, SoundWaveParam(nullptr)
-	{}
-
-	FAudioComponentParam()
-		: FloatParam(0.f)
-		, BoolParam(false)
-		, IntParam(0)
-		, SoundWaveParam(nullptr)
-	{
-	}
+	TObjectPtr<USoundWave> SoundWaveParam = nullptr;
 };
 
 /**
@@ -171,7 +146,7 @@ class UInitialActiveSoundParams : public UObject
 
 	// Collection of parameters to be sent to the active sound
 	UPROPERTY(EditAnywhere, BlueprintReadWrite, Category = "Audio")
-	TArray<FAudioComponentParam> AudioParams;
+	TArray<FAudioParameter> AudioParams;
 
 	void Reset(int32 ReserveSize = 0)
 	{
@@ -186,19 +161,20 @@ class UInitialActiveSoundParams : public UObject
  * @see USoundBase
  */
 UCLASS(ClassGroup=(Audio, Common), hidecategories=(Object, ActorComponent, Physics, Rendering, Mobility, LOD), ShowCategories=Trigger, meta=(BlueprintSpawnableComponent))
-class ENGINE_API UAudioComponent : public USceneComponent
+class ENGINE_API UAudioComponent : public USceneComponent, public ISoundGeneratorParameterInterface
 {
 	GENERATED_UCLASS_BODY()
 
 	/** The sound to be played */
 	UPROPERTY(EditAnywhere, BlueprintReadWrite, Category=Sound)
-	TObjectPtr<class USoundBase> Sound;
+	TObjectPtr<USoundBase> Sound;
 
-	/** Array of per-instance parameters for this AudioComponent. */
+	/** Array of per-instance parameters for this AudioComponent. Direct changes to this array from Blueprint (as opposed to updates
+	  * made via the Set functions) will not be forwarded to the sound if the component is actively playing. */
 	UPROPERTY(EditAnywhere, BlueprintReadWrite, Category=Sound, AdvancedDisplay)
-	TArray<struct FAudioComponentParam> InstanceParameters;
+	TArray<FAudioParameter> InstanceParameters;
 
-	/** Optional sound group this AudioComponent belongs to */
+	/** SoundClass that overrides that set on the referenced SoundBase when component is played. */
 	UPROPERTY(EditAnywhere, Category=Sound, AdvancedDisplay)
 	TObjectPtr<USoundClass> SoundClassOverride;
 
@@ -521,7 +497,7 @@ public:
 	 *  Doesn't indicate if the sound is paused or fading in/out. Use GetPlayState() to get the full play state.
 	 */
 	UFUNCTION(BlueprintCallable, Category="Audio|Components|Audio")
-	virtual bool IsPlaying() const;
+	virtual bool IsPlaying() const override;
 
 	/** Returns if the sound is virtualized. */
 	UFUNCTION(BlueprintCallable, Category="Audio|Components|Audio")
@@ -539,40 +515,14 @@ public:
 	UFUNCTION(BlueprintCallable, Category="Audio|Components|Audio")
 	void AdjustVolume(float AdjustVolumeDuration, float AdjustVolumeLevel, const EAudioFaderCurve FadeCurve = EAudioFaderCurve::Linear);
 
-	/** Allows the designer to set the Float Parameter on the SoundCue whose name matches the name indicated.
-	 * @param InName The name of the Float to set. It must match the name set in SoundCue's Crossfade By Param or Continuous Modulator Node.
-	 * @param InFloat The value to set the Parameter to.
-	 */
-	UFUNCTION(BlueprintCallable, Category="Audio|Components|Audio")
-	void SetFloatParameter(FName InName, float InFloat);
-
-	UPROPERTY(Transient)
-	TScriptInterface<IAudioParameterInterface> ParameterInterface;
-
-	/** Retrieves the parameter interface to set parameters on sound types that support parameterization (ex. MetaSounds if plugin is enabled). */
-	UFUNCTION(BlueprintCallable, Category = "Audio|Parameters")
-	TScriptInterface<IAudioParameterInterface> GetParameterInterface() const;
-
-	/** Allows the designer to set the Wave Parameter on the SoundCue whose name matches the name indicated.
-	 * @param InName The name of the Wave to set. It must match the name set in SoundCue's WaveParam Node
-	 * @param InWave The value to set the Parameter to
+	/** Sets the parameter matching the name indicated to the provided Wave. Provided for convenience/backward compatibility
+	 * with SoundCues (The parameter interface supports any object and is up to the system querying it to determine whether
+	 * it is a valid type).
+	 * @param InName The name of the parameter to assign the wave to.
+	 * @param InWave The wave value to set.
 	 */
-	UFUNCTION(BlueprintCallable, Category="Audio|Components|Audio")
-	void SetWaveParameter(FName InName, class USoundWave* InWave);
-
-	/** Allows the designer to set the Boolean Parameter on the SoundCue whose name matches the name indicated.
-	 * @param InName The name of the Boolean to set. It must match the name set in SoundCue's Branch Node 
-	 * @param InBool The value to set the Parameter to
-	 */
-	UFUNCTION(BlueprintCallable, Category="Audio|Components|Audio", meta=(DisplayName="Set Boolean Parameter"))
-	void SetBoolParameter(FName InName, bool InBool);
-
-	/** Allows the designer to set the Integer Parameter on the SoundCue whose name matches the name indicated.
-	 * @param InName The name of the Integer to set. It must match the name set in SoundCue's Switch Node
-	 * @param InInt The value to set the Parameter to
-	 */
-	UFUNCTION(BlueprintCallable, Category="Audio|Components|Audio", meta=(DisplayName="Set Integer Parameter"))
-	void SetIntParameter(FName InName, int32 InInt);
+	UFUNCTION(BlueprintCallable, Category = "Audio|Parameter")
+	void SetWaveParameter(FName InName, USoundWave* InWave);
 
 	/** Set a new volume multiplier */
 	UFUNCTION(BlueprintCallable, Category="Audio|Components|Audio")
@@ -701,10 +651,6 @@ private:
 	void BroadcastPlayState();
 
 public:
-
-	/** Sets the sound instance parameter. */
-	void SetSoundParameter(const FAudioComponentParam& Param);
-
 	/** Set when the sound is finished with initial fading in */
 	void SetFadeInComplete();
 
@@ -751,9 +697,6 @@ public:
 	 */
 	void CollectAttenuationShapesForVisualization(TMultiMap<EAttenuationShape::Type, FBaseAttenuationSettings::AttenuationShapeDetails>& ShapeDetailsMap) const;
 
-	/** Returns the active audio device to use for this component based on whether or not the component is playing in a world. */
-	FAudioDevice* GetAudioDevice() const;
-
 	uint64 GetAudioComponentID() const { return AudioComponentID; }
 
 	FName GetAudioComponentUserID() const { return AudioComponentUserID; }
@@ -765,6 +708,13 @@ public:
 	void SetPlaybackTimes(const TMap<uint32, float>& InSoundWavePlaybackTimes);
 
 	void SetSourceEffectChain(USoundEffectSourcePresetChain* InSourceEffectChain);
+
+	/** SoundGeneratorParameterInterface Implementation */
+	FAudioDevice* GetAudioDevice() const override;
+	TArray<FAudioParameter>& GetInstanceParameters() override { return InstanceParameters; }
+	uint64 GetInstanceID() const override { return AudioComponentID; }
+	USoundBase* GetSound() override { return Sound; }
+
 public:
 
 	/**
diff --git a/Engine/Source/Runtime/Engine/Classes/Sound/SoundBase.h b/Engine/Source/Runtime/Engine/Classes/Sound/SoundBase.h
--- a/Engine/Source/Runtime/Engine/Classes/Sound/SoundBase.h
+++ b/Engine/Source/Runtime/Engine/Classes/Sound/SoundBase.h
@@ -1,13 +1,11 @@
 // Copyright Epic Games, Inc. All Rights Reserved.
 #pragma once
 
-/**
- * The base class for a playable sound object
- */
 
 #include "Audio.h"
 #include "CoreMinimal.h"
 #include "IAudioExtensionPlugin.h"
+#include "IAudioParameterTransmitter.h"
 #include "Sound/AudioSettings.h"
 #include "Sound/SoundClass.h"
 #include "SoundConcurrency.h"
@@ -19,9 +17,16 @@
 #include "UObject/ObjectMacros.h"
 #include "AudioDeviceManager.h"
 #include "Interfaces/Interface_AssetUserData.h"
+
 #include "SoundBase.generated.h"
 
 
+// Forward Declarations
+namespace Audio
+{
+	class IParameterTransmitter;
+} // namespace Audio
+
 class USoundEffectSourcePreset;
 class USoundSourceBus;
 class USoundSubmix;
@@ -31,6 +36,7 @@ class UAssetUserData;
 struct FActiveSound;
 struct FSoundParseParameters;
 
+
 /**
  * Method of virtualization when a sound is stopped due to playback constraints
  * (i.e. by concurrency, priority, and/or MaxChannelCount)
@@ -54,6 +60,9 @@ enum class EVirtualizationMode : uint8
 	Restart
 };
 
+/**
+ * The base class for a playable sound object
+ */
 UCLASS(config=Engine, hidecategories=Object, abstract, editinlinenew, BlueprintType)
 class ENGINE_API USoundBase : public UObject, public IInterface_AssetUserData
 {
@@ -280,13 +289,18 @@ public:
 	virtual const TArray<UAssetUserData*>* GetAssetUserDataArray() const override;
 	//~ End IInterface_AssetUserData Interface
 
-	/** Called prior to attempting to play sound from Game Thread to enable initializing any additional resources required for playback. */
+	/** Called from the Game Thread prior to attempting to pass parameters to the ParameterTransmitter. */
+	virtual void InitParameters(TArray<FAudioParameter>& InParametersToInit, FName InFeatureName);
+
+	/** Called from the Game Thread prior to attempting to initialize a sound instance. */
 	virtual void InitResources() { }
 
 	/** Creates a sound generator instance from this sound base. Return true if this is being implemented by a subclass. Sound generators procedurally generate audio in the audio render thread. */
 	virtual ISoundGeneratorPtr CreateSoundGenerator(const FSoundGeneratorInitParams& InParams) { return nullptr; }
 
-	/** Creates a instance transmitter for communicating with active sound instances. */
-	virtual TUniquePtr<IAudioInstanceTransmitter> CreateInstanceTransmitter(const FAudioInstanceTransmitterInitParams& InParams) const { return nullptr; }
-};
+	/** Creates a parameter transmitter for communicating with active sound instances. */
+	virtual TUniquePtr<Audio::IParameterTransmitter> CreateParameterTransmitter(Audio::FParameterTransmitterInitParams&& InParams) const;
 
+	/** Returns whether parameter is valid input for the given sound */
+	virtual bool IsParameterValid(const FAudioParameter& InParameter) const;
+};
diff --git a/Engine/Source/Runtime/Engine/Classes/Sound/SoundGenerator.h b/Engine/Source/Runtime/Engine/Classes/Sound/SoundGenerator.h
--- a/Engine/Source/Runtime/Engine/Classes/Sound/SoundGenerator.h
+++ b/Engine/Source/Runtime/Engine/Classes/Sound/SoundGenerator.h
@@ -15,6 +15,7 @@ struct FSoundGeneratorInitParams
 	int32 NumFramesPerCallback = 0;
 	uint64 InstanceID = 0;
 	bool bIsPreviewSound = false;
+	FString GraphName;
 };
 
 class ENGINE_API ISoundGenerator
diff --git a/Engine/Source/Runtime/Engine/Private/ActiveSound.cpp b/Engine/Source/Runtime/Engine/Private/ActiveSound.cpp
--- a/Engine/Source/Runtime/Engine/Private/ActiveSound.cpp
+++ b/Engine/Source/Runtime/Engine/Private/ActiveSound.cpp
@@ -224,11 +224,12 @@ void FActiveSound::AddReferencedObjects(FReferenceCollector& Collector)
 		}
 	}
 
-	for (FAudioComponentParam& Param : InstanceParameters)
+	if (InstanceTransmitter.IsValid())
 	{
-		if (Param.SoundWaveParam)
+		TArray<UObject*> InstanceReferences = InstanceTransmitter->GetReferencedObjects();
+		for (UObject* Object : InstanceReferences)
 		{
-			Collector.AddReferencedObject(Param.SoundWaveParam);
+			Collector.AddReferencedObject(Object);
 		}
 	}
 }
@@ -395,7 +396,7 @@ void FActiveSound::SetSourceBusSend(EBusSendType BusSendType, const FSoundSource
 			Info.SendLevel = SendInfo.SendLevel;
 
 			bHasNewBusSends = true;
-			newBusSends.Add(TTuple<EBusSendType, FSoundSourceBusSendInfo>(BusSendType, SendInfo));
+			NewBusSends.Add(TTuple<EBusSendType, FSoundSourceBusSendInfo>(BusSendType, SendInfo));
 			return;
 		}
 	}
@@ -404,7 +405,7 @@ void FActiveSound::SetSourceBusSend(EBusSendType BusSendType, const FSoundSource
 	BusSendsOverride[(int32)BusSendType].Add(SendInfo);
 
 	bHasNewBusSends = true;
-	newBusSends.Add(TTuple<EBusSendType, FSoundSourceBusSendInfo>(BusSendType,SendInfo));
+	NewBusSends.Add(TTuple<EBusSendType, FSoundSourceBusSendInfo>(BusSendType,SendInfo));
 }
 
 bool FActiveSound::HasNewBusSends() const
@@ -414,12 +415,12 @@ bool FActiveSound::HasNewBusSends() const
 
 TArray< TTuple<EBusSendType, FSoundSourceBusSendInfo> > const& FActiveSound::GetNewBusSends() const
 {
-	return newBusSends;
+	return NewBusSends;
 }
 
 void FActiveSound::ResetNewBusSends()
 {
-	newBusSends.Empty();
+	NewBusSends.Empty();
 	bHasNewBusSends = false;
 }
 
@@ -1263,24 +1264,6 @@ void FActiveSound::ApplyRadioFilter(const FSoundParseParameters& ParseParams)
 	bRadioFilterSelected = true;
 }
 
-bool FActiveSound::GetFloatParameter(const FName InName, float& OutFloat) const
-{
-	// Always fail if we pass in no name.
-	if (InName != NAME_None)
-	{
-		for (const FAudioComponentParam& P : InstanceParameters)
-		{
-			if (P.ParamName == InName)
-			{
-				OutFloat = P.FloatParam;
-				return true;
-			}
-		}
-	}
-
-	return false;
-}
-
 float FActiveSound::GetVolume() const
 {
 	const float Volume = VolumeMultiplier * ComponentVolumeFader.GetVolume() * GetTotalConcurrencyVolumeScale();
@@ -1307,163 +1290,6 @@ void FActiveSound::UpdateConcurrencyVolumeScalars(const float DeltaTime)
 	}
 }
 
-void FActiveSound::SetFloatParameter(const FName InName, const float InFloat)
-{
-	if (InName != NAME_None)
-	{		
-		// First see if an entry for this name already exists
-		for (FAudioComponentParam& P : InstanceParameters)
-		{
-			if (P.ParamName == InName)
-			{
-				P.FloatParam = InFloat;
-				return;
-			}
-		}
-
-		// We didn't find one, so create a new one.
-		const int32 NewParamIndex = InstanceParameters.AddDefaulted();
-		InstanceParameters[NewParamIndex].ParamName = InName;
-		InstanceParameters[NewParamIndex].FloatParam = InFloat;
-	}
-}
-
-bool FActiveSound::GetWaveParameter(const FName InName, USoundWave*& OutWave) const
-{
-	// Always fail if we pass in no name.
-	if (InName != NAME_None)
-	{
-		for (const FAudioComponentParam& P : InstanceParameters)
-		{
-			if (P.ParamName == InName)
-			{
-				OutWave = P.SoundWaveParam;
-				return true;
-			}
-		}
-	}
-
-	return false;
-}
-
-void FActiveSound::SetWaveParameter(const FName InName, USoundWave* InWave)
-{
-	if (InName != NAME_None)
-	{
-		// First see if an entry for this name already exists
-		for (FAudioComponentParam& P : InstanceParameters)
-		{
-			if (P.ParamName == InName)
-			{
-				P.SoundWaveParam = InWave;
-				return;
-			}
-		}
-
-		// We didn't find one, so create a new one.
-		const int32 NewParamIndex = InstanceParameters.AddDefaulted();
-		InstanceParameters[NewParamIndex].ParamName = InName;
-		InstanceParameters[NewParamIndex].SoundWaveParam = InWave;
-	}
-}
-
-bool FActiveSound::GetBoolParameter(const FName InName, bool& OutBool) const
-{
-	// Always fail if we pass in no name.
-	if (InName != NAME_None)
-	{
-		for (const FAudioComponentParam& P : InstanceParameters)
-		{
-			if (P.ParamName == InName)
-			{
-				OutBool = P.BoolParam;
-				return true;
-			}
-		}
-	}
-
-	return false;
-}
-
-void FActiveSound::SetBoolParameter(const FName InName, const bool InBool)
-{
-	if (InName != NAME_None)
-	{
-		// First see if an entry for this name already exists
-		for (FAudioComponentParam& P : InstanceParameters)
-		{
-			if (P.ParamName == InName)
-			{
-				P.BoolParam = InBool;
-				return;
-			}
-		}
-
-		// We didn't find one, so create a new one.
-		const int32 NewParamIndex = InstanceParameters.AddDefaulted();
-		InstanceParameters[NewParamIndex].ParamName = InName;
-		InstanceParameters[NewParamIndex].BoolParam = InBool;
-	}
-}
-
-bool FActiveSound::GetIntParameter(const FName InName, int32& OutInt) const
-{
-	// Always fail if we pass in no name.
-	if (InName != NAME_None)
-	{
-		for (const FAudioComponentParam& P : InstanceParameters)
-		{
-			if (P.ParamName == InName)
-			{
-				OutInt = P.IntParam;
-				return true;
-			}
-		}
-	}
-
-	return false;
-}
-
-void FActiveSound::SetIntParameter(const FName InName, const int32 InInt)
-{
-	if (InName != NAME_None)
-	{
-		// First see if an entry for this name already exists
-		for (FAudioComponentParam& P : InstanceParameters)
-		{
-			if (P.ParamName == InName)
-			{
-				P.IntParam = InInt;
-				return;
-			}
-		}
-
-		// We didn't find one, so create a new one.
-		const int32 NewParamIndex = InstanceParameters.AddDefaulted();
-		InstanceParameters[NewParamIndex].ParamName = InName;
-		InstanceParameters[NewParamIndex].IntParam = InInt;
-	}
-}
-
-void FActiveSound::SetSoundParameter(const FAudioComponentParam& Param)
-{
-	if (Param.ParamName != NAME_None)
-	{
-		// First see if an entry for this name already exists
-		for (FAudioComponentParam& P : InstanceParameters)
-		{
-			if (P.ParamName == Param.ParamName)
-			{
-				P = Param;
-				return;
-			}
-		}
-
-		// We didn't find one, so create a new one.
-		const int32 NewParamIndex = InstanceParameters.Add(Param);
-	}
-}
-
 void FActiveSound::CollectAttenuationShapesForVisualization(TMultiMap<EAttenuationShape::Type, FBaseAttenuationSettings::AttenuationShapeDetails>& ShapeDetailsMap) const
 {
 	bool bFoundAttenuationSettings = false;
diff --git a/Engine/Source/Runtime/Engine/Private/Audio/AudioComponentParameterization.cpp b/Engine/Source/Runtime/Engine/Private/Audio/AudioComponentParameterization.cpp
deleted file mode 100644
diff --git a/Engine/Source/Runtime/Engine/Private/Audio/AudioParameterInterface.cpp b/Engine/Source/Runtime/Engine/Private/Audio/AudioParameterInterface.cpp
deleted file mode 100644
diff --git a/Engine/Source/Runtime/Engine/Private/Audio/SoundGeneratorParameterInterface.cpp b/Engine/Source/Runtime/Engine/Private/Audio/SoundGeneratorParameterInterface.cpp
new file mode 100644
--- /dev/null
+++ b/Engine/Source/Runtime/Engine/Private/Audio/SoundGeneratorParameterInterface.cpp
@@ -0,0 +1,193 @@
+// Copyright Epic Games, Inc. All Rights Reserved.
+#include "Audio/SoundGeneratorParameterInterface.h"
+
+#include "ActiveSound.h"
+#include "Audio.h"
+#include "AudioDevice.h"
+#include "AudioThread.h"
+#include "Components/AudioComponent.h"
+#include "IAudioExtensionPlugin.h"
+
+
+namespace SoundGeneratorParameterInterfacePrivate
+{
+	static const FName ProxyFeatureName("SoundGeneratorParameterInterface");
+} // namespace SoundGeneratorParameterInterfacePrivate
+
+USoundGeneratorParameterInterface::USoundGeneratorParameterInterface(FObjectInitializer const& InObjectInitializer)
+	: Super(InObjectInitializer)
+{
+}
+
+void ISoundGeneratorParameterInterface::ResetParameters()
+{
+	if (FAudioDevice* AudioDevice = GetAudioDevice())
+	{
+		FAudioThread::RunCommandOnAudioThread([AudioDevice, InstanceID = GetInstanceID()]() mutable
+		{
+			if (FActiveSound* ActiveSound = AudioDevice->FindActiveSound(InstanceID))
+			{
+				if (Audio::IParameterTransmitter* Transmitter = ActiveSound->GetTransmitter())
+				{
+					Transmitter->Reset();
+				}
+			}
+		});
+	}
+}
+
+void ISoundGeneratorParameterInterface::SetTriggerParameter(FName InName)
+{
+	// Trigger is just a (true) bool param currently.
+	SetParameterInternal(FAudioParameter(InName, true));
+}
+
+void ISoundGeneratorParameterInterface::SetBoolParameter(FName InName, bool InValue)
+{
+	SetParameterInternal(FAudioParameter(InName, InValue));
+}
+
+void ISoundGeneratorParameterInterface::SetBoolArrayParameter(FName InName, const TArray<bool>& InValue)
+{
+	SetParameterInternal(FAudioParameter(InName, InValue));
+}
+
+void ISoundGeneratorParameterInterface::SetIntParameter(FName InName, int32 InValue)
+{
+	SetParameterInternal(FAudioParameter(InName, InValue));
+}
+
+void ISoundGeneratorParameterInterface::SetIntArrayParameter(FName InName, const TArray<int32>& InValue)
+{
+	SetParameterInternal(FAudioParameter(InName, InValue));
+}
+
+void ISoundGeneratorParameterInterface::SetFloatParameter(FName InName, float InValue)
+{
+	SetParameterInternal(FAudioParameter(InName, InValue));
+}
+
+void ISoundGeneratorParameterInterface::SetFloatArrayParameter(FName InName, const TArray<float>& InValue)
+{
+	SetParameterInternal(FAudioParameter(InName, InValue));
+}
+
+void ISoundGeneratorParameterInterface::SetStringParameter(FName InName, const FString& InValue)
+{
+	SetParameterInternal(FAudioParameter(InName, InValue));
+}
+
+void ISoundGeneratorParameterInterface::SetStringArrayParameter(FName InName, const TArray<FString>& InValue)
+{
+	SetParameterInternal(FAudioParameter(InName, InValue));
+}
+
+void ISoundGeneratorParameterInterface::SetObjectParameter(FName InName, UObject* InValue)
+{
+	SetParameterInternal(FAudioParameter(InName, InValue));
+}
+
+void ISoundGeneratorParameterInterface::SetObjectArrayParameter(FName InName, const TArray<UObject*>& InValue)
+{
+	SetParameterInternal(FAudioParameter(InName, InValue));
+}
+
+void ISoundGeneratorParameterInterface::SetParameter(FAudioParameter&& InValue)
+{
+	SetParameterInternal(MoveTemp(InValue));
+}
+
+void ISoundGeneratorParameterInterface::SetParameters(TArray<FAudioParameter>&& InValues)
+{
+	for (FAudioParameter& Value : InValues)
+	{
+		TArray<FAudioParameter>& InstanceParameters = GetInstanceParameters();
+		if (FAudioParameter* CurrentParam = FAudioParameter::FindOrAddParam(InstanceParameters, Value.ParamName))
+		{
+			CurrentParam->Merge(Value, false /* bInTakeName */);
+		}
+	}
+
+	if (IsPlaying())
+	{
+		if (FAudioDevice* AudioDevice = GetAudioDevice())
+		{
+			TArray<FAudioParameter> ParamsToSet;
+			if (USoundBase* Sound = GetSound())
+			{
+				Sound->InitParameters(InValues, SoundGeneratorParameterInterfacePrivate::ProxyFeatureName);
+			}
+
+			ParamsToSet = MoveTemp(InValues);
+
+			DECLARE_CYCLE_STAT(TEXT("FAudioThreadTask.SoundGenerator.SetParameters"), STAT_AudioSetParameters, STATGROUP_AudioThreadCommands);
+
+			FAudioThread::RunCommandOnAudioThread([AudioDevice, InstanceID = GetInstanceID(), Params { MoveTemp(ParamsToSet) }]() mutable
+			{
+				if (FActiveSound* ActiveSound = AudioDevice->FindActiveSound(InstanceID))
+				{
+					if (Audio::IParameterTransmitter* Transmitter = ActiveSound->GetTransmitter())
+					{
+						for (FAudioParameter& Param : Params)
+						{
+							const FName ParamName = Param.ParamName;
+							if (!Transmitter->SetParameter(MoveTemp(Param)))
+							{
+								UE_LOG(LogAudio, Warning, TEXT("Failed to set parameter '%s'"), *ParamName.ToString());
+							}
+						}
+					}
+				}
+			}, GET_STATID(STAT_AudioSetParameters));
+		}
+	}
+}
+
+void ISoundGeneratorParameterInterface::SetParameterInternal(FAudioParameter&& InParam)
+{
+	if (!InParam.ParamName.IsValid())
+	{
+		return;
+	}
+
+	TArray<FAudioParameter>& InstanceParameters = GetInstanceParameters();
+	if (FAudioParameter* CurrentParam = FAudioParameter::FindOrAddParam(InstanceParameters, InParam.ParamName))
+	{
+		CurrentParam->Merge(InParam, false /* bInTakeName */);
+	}
+
+	if (IsPlaying())
+	{
+		if (FAudioDevice* AudioDevice = GetAudioDevice())
+		{
+			FAudioParameter ParamToSet;
+			if (USoundBase* Sound = GetSound())
+			{
+				TArray<FAudioParameter> Params = { MoveTemp(InParam) };
+				Sound->InitParameters(Params, SoundGeneratorParameterInterfacePrivate::ProxyFeatureName);
+				ParamToSet = MoveTemp(Params[0]);
+			}
+			else
+			{
+				ParamToSet = MoveTemp(InParam);
+			}
+
+			DECLARE_CYCLE_STAT(TEXT("FAudioThreadTask.SoundGenerator.SetParameter"), STAT_AudioSetParameter, STATGROUP_AudioThreadCommands);
+
+			FAudioThread::RunCommandOnAudioThread([AudioDevice, InstanceID = GetInstanceID(), Param { MoveTemp(ParamToSet) }]() mutable
+			{
+				if (FActiveSound* ActiveSound = AudioDevice->FindActiveSound(InstanceID))
+				{
+					if (Audio::IParameterTransmitter* Transmitter = ActiveSound->GetTransmitter())
+					{
+						const FName ParamName = Param.ParamName;
+						if (!Transmitter->SetParameter(MoveTemp(Param)))
+						{
+							UE_LOG(LogAudio, Warning, TEXT("Failed to set parameter '%s'"), *ParamName.ToString());
+						}
+					}
+				}
+			}, GET_STATID(STAT_AudioSetParameter));
+		}
+	}
+}
diff --git a/Engine/Source/Runtime/Engine/Private/AudioDevice.cpp b/Engine/Source/Runtime/Engine/Private/AudioDevice.cpp
--- a/Engine/Source/Runtime/Engine/Private/AudioDevice.cpp
+++ b/Engine/Source/Runtime/Engine/Private/AudioDevice.cpp
@@ -4677,17 +4677,29 @@ void FAudioDevice::NotifyPluginListenersWorldChanged(UWorld* World)
 	}
 }
 
-void FAudioDevice::AddNewActiveSound(const FActiveSound& NewActiveSound)
+void FAudioDevice::AddNewActiveSound(const FActiveSound& ActiveSound, const TArray<FAudioParameter>* InDefaultParams)
+{
+	TArray<FAudioParameter> Params;
+	if (InDefaultParams)
+	{
+		Params.Append(*InDefaultParams);
+	}
+
+	AddNewActiveSound(ActiveSound, MoveTemp(Params));
+}
+
+void FAudioDevice::AddNewActiveSound(const FActiveSound& NewActiveSound, TArray<FAudioParameter>&& InDefaultParams)
 {
 	if (USoundBase* Sound = NewActiveSound.GetSound())
 	{
 		Sound->InitResources();
+		Sound->InitParameters(InDefaultParams, "NewSoundRequest");
 	}
 
-	AddNewActiveSoundInternal(NewActiveSound, nullptr);
+	AddNewActiveSoundInternal(NewActiveSound, MoveTemp(InDefaultParams));
 }
 
-void FAudioDevice::AddNewActiveSoundInternal(const FActiveSound& NewActiveSound, FAudioVirtualLoop* VirtualLoopToRetrigger)
+void FAudioDevice::AddNewActiveSoundInternal(const FActiveSound& InNewActiveSound, TArray<FAudioParameter>&& InDefaultParams, FAudioVirtualLoop* InVirtualLoopToRetrigger)
 {
 	LLM_SCOPE(ELLMTag::AudioMisc);
 	
@@ -4695,29 +4707,29 @@ void FAudioDevice::AddNewActiveSoundInternal(const FActiveSound& NewActiveSound,
 	{
 		DECLARE_CYCLE_STAT(TEXT("FAudioThreadTask.AddNewActiveSound"), STAT_AudioAddNewActiveSound, STATGROUP_AudioThreadCommands);
 
-		FAudioDevice* AudioDevice = this;
-		FAudioThread::RunCommandOnAudioThread([AudioDevice, NewActiveSound, VirtualLoopToRetrigger]()
+
+		FAudioThread::RunCommandOnAudioThread([AudioDevice = this, InNewActiveSound, DefaultParams = MoveTemp(InDefaultParams)]() mutable
 		{
-			AudioDevice->AddNewActiveSoundInternal(NewActiveSound, VirtualLoopToRetrigger);
+			AudioDevice->AddNewActiveSoundInternal(InNewActiveSound, MoveTemp(DefaultParams));
 		}, GET_STATID(STAT_AudioAddNewActiveSound));
 
 		return;
 	}
 
-	USoundBase* Sound = NewActiveSound.GetSound();
+	USoundBase* Sound = InNewActiveSound.GetSound();
 	if (Sound == nullptr)
 	{
-		ReportSoundFailedToStart(NewActiveSound.AudioComponentID, VirtualLoopToRetrigger);
+		ReportSoundFailedToStart(InNewActiveSound.AudioComponentID, InVirtualLoopToRetrigger);
 		return;
 	}
 
 	// Don't allow buses to try to play if we're not using the audio mixer.
 	if (!IsAudioMixerEnabled())
 	{
-		USoundSourceBus* Bus = Cast<USoundSourceBus>(NewActiveSound.Sound);
+		USoundSourceBus* Bus = Cast<USoundSourceBus>(InNewActiveSound.Sound);
 		if (Bus)
 		{
-			ReportSoundFailedToStart(NewActiveSound.AudioComponentID, VirtualLoopToRetrigger);
+			ReportSoundFailedToStart(InNewActiveSound.AudioComponentID, InVirtualLoopToRetrigger);
 			return;
 		}
 	}
@@ -4726,11 +4738,11 @@ void FAudioDevice::AddNewActiveSoundInternal(const FActiveSound& NewActiveSound,
 	{
 		// TODO: Determine if this check has already been completed at AudioComponent level and skip if so. Also,
 		// unify code paths determining if sound is audible.
-		if (!SoundIsAudible(NewActiveSound))
+		if (!SoundIsAudible(InNewActiveSound))
 		{
-			UE_LOG(LogAudio, Log, TEXT("New ActiveSound not created for out of range Sound %s"), *NewActiveSound.Sound->GetName());
+			UE_LOG(LogAudio, Log, TEXT("New ActiveSound not created for out of range Sound %s"), *InNewActiveSound.Sound->GetName());
 
-			ReportSoundFailedToStart(NewActiveSound.AudioComponentID, VirtualLoopToRetrigger);
+			ReportSoundFailedToStart(InNewActiveSound.AudioComponentID, InVirtualLoopToRetrigger);
 			return;
 		}
 	}
@@ -4745,10 +4757,10 @@ void FAudioDevice::AddNewActiveSoundInternal(const FActiveSound& NewActiveSound,
 		{
 			// Reject the new sound if it doesn't have the debug sound name substring
 			FString SoundName;
-			NewActiveSound.Sound->GetName(SoundName);
+			InNewActiveSound.Sound->GetName(SoundName);
 			if (!SoundName.Contains(DebugSound))
 			{
-				ReportSoundFailedToStart(NewActiveSound.AudioComponentID, VirtualLoopToRetrigger);
+				ReportSoundFailedToStart(InNewActiveSound.AudioComponentID, InVirtualLoopToRetrigger);
 				return;
 			}
 		}
@@ -4756,11 +4768,11 @@ void FAudioDevice::AddNewActiveSoundInternal(const FActiveSound& NewActiveSound,
 #endif // !UE_BUILD_SHIPPING
 
 	// Determine if sound is loop and eligible for virtualize prior to creating "live" active sound in next Concurrency check step
-	if (!VirtualLoopToRetrigger)
+	if (!InVirtualLoopToRetrigger)
 	{
 		const bool bDoRangeCheck = true;
 		FAudioVirtualLoop VirtualLoop;
-		if (FAudioVirtualLoop::Virtualize(NewActiveSound, *this, bDoRangeCheck, VirtualLoop))
+		if (FAudioVirtualLoop::Virtualize(InNewActiveSound, *this, bDoRangeCheck, VirtualLoop))
 		{
 			UE_LOG(LogAudio, Verbose, TEXT("New ActiveSound %s Virtualizing: Failed to pass initial audible range check"), *Sound->GetName());
 			AddVirtualLoop(VirtualLoop);
@@ -4768,31 +4780,31 @@ void FAudioDevice::AddNewActiveSoundInternal(const FActiveSound& NewActiveSound,
 		}
 	}
 
-	// Evaluate concurrency. This will create an ActiveSound ptr which is a copy of NewActiveSound if the sound can play.
+	// Evaluate concurrency. This will create an ActiveSound ptr which is a copy of InNewActiveSound if the sound can play.
 	FActiveSound* ActiveSound = nullptr;
 
 	{
 		SCOPE_CYCLE_COUNTER(STAT_AudioEvaluateConcurrency);
 
 		// Try to create a new active sound. This returns nullptr if too many sounds are playing with this sound's concurrency setting
-		ActiveSound = ConcurrencyManager.CreateNewActiveSound(NewActiveSound, VirtualLoopToRetrigger != nullptr);
+		ActiveSound = ConcurrencyManager.CreateNewActiveSound(InNewActiveSound, InVirtualLoopToRetrigger != nullptr);
 	}
 
 	// Didn't pass concurrency, and not an attempt to revive from virtualization, so see if candidate for virtualization
 	if (!ActiveSound)
 	{
-		if (!VirtualLoopToRetrigger)
+		if (!InVirtualLoopToRetrigger)
 		{
 			const bool bDoRangeCheck = false;
 			FAudioVirtualLoop VirtualLoop;
-			if (FAudioVirtualLoop::Virtualize(NewActiveSound, *this, bDoRangeCheck, VirtualLoop))
+			if (FAudioVirtualLoop::Virtualize(InNewActiveSound, *this, bDoRangeCheck, VirtualLoop))
 			{
 				UE_LOG(LogAudioConcurrency, Verbose, TEXT("New ActiveSound %s Virtualizing: Failed to pass concurrency"), *Sound->GetName());
 				AddVirtualLoop(VirtualLoop);
 			}
 			else
 			{
-				ReportSoundFailedToStart(NewActiveSound.GetAudioComponentID(), VirtualLoopToRetrigger);
+				ReportSoundFailedToStart(InNewActiveSound.GetAudioComponentID(), InVirtualLoopToRetrigger);
 			}
 		}
 		return;
@@ -4803,7 +4815,7 @@ void FAudioDevice::AddNewActiveSoundInternal(const FActiveSound& NewActiveSound,
 	if (GIsEditor)
 	{
 		// If the sound played on an editor preview world, treat it as a preview sound (unpausable and ignoring the realtime volume slider)
-		if (const UWorld* World = NewActiveSound.GetWorld())
+		if (const UWorld* World = InNewActiveSound.GetWorld())
 		{
 			ActiveSound->bIsPreviewSound |= (World->WorldType == EWorldType::EditorPreview);
 		}
@@ -4817,7 +4829,7 @@ void FAudioDevice::AddNewActiveSoundInternal(const FActiveSound& NewActiveSound,
 	(*PlayCount)++;
 
 #if !(UE_BUILD_SHIPPING || UE_BUILD_TEST)
-	UE_LOG(LogAudio, VeryVerbose, TEXT("New ActiveSound %s Comp: %s Loc: %s"), *Sound->GetName(), *NewActiveSound.GetAudioComponentName(), *NewActiveSound.Transform.GetTranslation().ToString());
+	UE_LOG(LogAudio, VeryVerbose, TEXT("New ActiveSound %s Comp: %s Loc: %s"), *Sound->GetName(), *InNewActiveSound.GetAudioComponentName(), *InNewActiveSound.Transform.GetTranslation().ToString());
 #endif // !(UE_BUILD_SHIPPING || UE_BUILD_TEST)
 
 	// Cull one-shot active sounds if we've reached our max limit of one shot active sounds before we attempt to evaluate concurrency
@@ -4844,23 +4856,33 @@ void FAudioDevice::AddNewActiveSoundInternal(const FActiveSound& NewActiveSound,
 	// Mark to not report playback complete on destruction as responsibility therein has been
 	// passed to newly created ActiveSound added below.  Add as stopping sound prior to adding
 	// new sound to ensure AudioComponentIDToActiveSoundMap is registered with the correct ActiveSound.
-	if (VirtualLoopToRetrigger)
+	if (InVirtualLoopToRetrigger)
 	{
-		FActiveSound& VirtualActiveSound = VirtualLoopToRetrigger->GetActiveSound();
+		FActiveSound& VirtualActiveSound = InVirtualLoopToRetrigger->GetActiveSound();
 		AddSoundToStop(&VirtualActiveSound);
 
 		// Clear must be called after AddSoundToStop to ensure AudioComponent is properly removed from AudioComponentIDToActiveSoundMap
 		VirtualActiveSound.ClearAudioComponent();
 	}
 
-	FAudioInstanceTransmitterInitParams TransmitterInitParams
+	// Retriggering a virtualized ActiveSound results in an
+	// already created, valid transmitter, and thus forwarded
+	// initial params should be ignored to avoid stomping existing
+	// parameter state.
+	if (!ActiveSound->InstanceTransmitter.IsValid())
 	{
-		ActiveSound->GetAudioComponentID(),
-		GetSampleRate()
-	};
-	ActiveSound->InstanceTransmitter = Sound->CreateInstanceTransmitter(TransmitterInitParams);
+		Audio::FParameterTransmitterInitParams TransmitterInitParams
+		{
+			ActiveSound->GetAudioComponentID(),
+			GetSampleRate(),
+			MoveTemp(InDefaultParams)
+		};
+
+		ActiveSound->InstanceTransmitter = Sound->CreateParameterTransmitter(MoveTemp(TransmitterInitParams));
+	}
 
 	ActiveSounds.Add(ActiveSound);
+
 	if (ActiveSound->GetAudioComponentID() > 0)
 	{
 		AudioComponentIDToActiveSoundMap.Add(ActiveSound->GetAudioComponentID(), ActiveSound);
@@ -4887,7 +4909,7 @@ void FAudioDevice::RetriggerVirtualLoop(FAudioVirtualLoop& VirtualLoopToRetrigge
 {
 	check(IsInAudioThread());
 
-	AddNewActiveSoundInternal(VirtualLoopToRetrigger.GetActiveSound(), &VirtualLoopToRetrigger);
+	AddNewActiveSoundInternal(VirtualLoopToRetrigger.GetActiveSound(), { }, &VirtualLoopToRetrigger);
 }
 
 void FAudioDevice::AddEnvelopeFollowerDelegate(USoundSubmix* InSubmix, const FOnSubmixEnvelopeBP& OnSubmixEnvelopeBP)
@@ -5855,7 +5877,7 @@ UAudioComponent* FAudioDevice::CreateComponent(USoundBase* Sound, const FCreateC
 	return AudioComponent;
 }
 
-void FAudioDevice::PlaySoundAtLocation(USoundBase* Sound, UWorld* World, float VolumeMultiplier, float PitchMultiplier, float StartTime, const FVector& Location, const FRotator& Rotation, USoundAttenuation* AttenuationSettings, USoundConcurrency* Concurrency, const TArray<FAudioComponentParam>* Params, AActor* OwningActor)
+void FAudioDevice::PlaySoundAtLocation(USoundBase* Sound, UWorld* World, float VolumeMultiplier, float PitchMultiplier, float StartTime, const FVector& Location, const FRotator& Rotation, USoundAttenuation* AttenuationSettings, USoundConcurrency* Concurrency, const TArray<FAudioParameter>* Params, AActor* OwningActor)
 {
 	check(IsInGameThread());
 
@@ -5913,16 +5935,7 @@ void FAudioDevice::PlaySoundAtLocation(USoundBase* Sound, UWorld* World, float V
 
 		NewActiveSound.SetOwner(OwningActor);
 
-		// Apply any optional audio component instance params on the sound
-		if (Params)
-		{
-			for (const FAudioComponentParam& Param : *Params)
-			{
-				NewActiveSound.SetSoundParameter(Param);
-			}
-		}
-
-		AddNewActiveSound(NewActiveSound);
+		AddNewActiveSound(NewActiveSound, Params);
 	}
 	else
 	{
diff --git a/Engine/Source/Runtime/Engine/Private/Components/AudioComponent.cpp b/Engine/Source/Runtime/Engine/Private/Components/AudioComponent.cpp
--- a/Engine/Source/Runtime/Engine/Private/Components/AudioComponent.cpp
+++ b/Engine/Source/Runtime/Engine/Private/Components/AudioComponent.cpp
@@ -3,7 +3,7 @@
 #include "Components/AudioComponent.h"
 #include "ActiveSound.h"
 #include "Audio.h"
-#include "Audio/AudioComponentParameterization.h"
+#include "Audio/SoundGeneratorParameterInterface.h"
 #include "AudioDevice.h"
 #include "AudioThread.h"
 #include "Components/BillboardComponent.h"
@@ -33,21 +33,16 @@ FAutoConsoleVariableRef CVarPrimeSoundOnAudioComponentSpawn(
 	ECVF_Default);
 
 
-/*-----------------------------------------------------------------------------
-UInitialActiveSoundParams implementation.
------------------------------------------------------------------------------*/
+uint64 UAudioComponent::AudioComponentIDCounter = 0;
+TMap<uint64, UAudioComponent*> UAudioComponent::AudioIDToComponentMap;
+FCriticalSection UAudioComponent::AudioIDToComponentMapLock;
+
+
 UInitialActiveSoundParams::UInitialActiveSoundParams(const FObjectInitializer& ObjectInitializer)
 	: Super(ObjectInitializer)
 {
 }
 
-/*-----------------------------------------------------------------------------
-UAudioComponent implementation.
------------------------------------------------------------------------------*/
-uint64 UAudioComponent::AudioComponentIDCounter = 0;
-TMap<uint64, UAudioComponent*> UAudioComponent::AudioIDToComponentMap;
-FCriticalSection UAudioComponent::AudioIDToComponentMapLock;
-
 UAudioComponent::UAudioComponent(const FObjectInitializer& ObjectInitializer)
 	: Super(ObjectInitializer)
 {
@@ -107,17 +102,14 @@ UAudioComponent* UAudioComponent::GetAudioComponentFromID(uint64 AudioComponentI
 
 void UAudioComponent::BeginDestroy()
 {
-	if (ParameterInterface)
-	{
-		ParameterInterface->Shutdown();
-	}
-
 	if (IsActive() && Sound && Sound->IsLooping())
 	{
 		UE_LOG(LogAudio, Verbose, TEXT("Audio Component is being destroyed prior to stopping looping sound '%s' directly."), *Sound->GetFullName());
 		Stop();
 	}
 
+	ResetParameters();
+
 	{
 		FScopeLock Lock(&AudioIDToComponentMapLock);
 		AudioIDToComponentMap.Remove(AudioComponentID);
@@ -512,12 +504,6 @@ void UAudioComponent::PlayInternal(const PlayInternalRequestData& InPlayRequestD
 
 			AudioDevice->GetMaxDistanceAndFocusFactor(Sound, World, Location, AttenuationSettingsToApply, MaxDistance, FocusFactor);
 
-			if (!ParameterInterface)
-			{
-				static const FName InterfaceName("AudioComponentParameterization");
-				ParameterInterface = NewObject<UAudioComponentParameterization>(this, InterfaceName);
-			}
-
 			FActiveSound NewActiveSound;
 			NewActiveSound.SetAudioComponent(*this);
 			NewActiveSound.SetWorld(GetWorld());
@@ -609,7 +595,6 @@ void UAudioComponent::PlayInternal(const PlayInternalRequestData& InPlayRequestD
 			NewActiveSound.QuantizedRequestData = InPlayRequestData.QuantizedRequestData;
 
 			NewActiveSound.MaxDistance = MaxDistance;
-			NewActiveSound.InstanceParameters = InstanceParameters;
 
 			Audio::FVolumeFader& Fader = NewActiveSound.ComponentVolumeFader;
 			Fader.SetVolume(0.0f); // Init to 0.0f to fade as default is 1.0f
@@ -617,7 +602,8 @@ void UAudioComponent::PlayInternal(const PlayInternalRequestData& InPlayRequestD
 
 			// Bump ActiveCount... this is used to determine if an audio component is still active after a sound reports back as completed
 			++ActiveCount;
-			AudioDevice->AddNewActiveSound(NewActiveSound);
+
+			AudioDevice->AddNewActiveSound(NewActiveSound, &InstanceParameters);
 
 			// In editor, the audio thread is not run separate from the game thread, and can result in calling PlaybackComplete prior
 			// to bIsActive being set. Therefore, we assign to the current state of ActiveCount as opposed to just setting to true.
@@ -771,11 +757,7 @@ void UAudioComponent::Stop()
 
 	if (bIsPreviewSound)
 	{
-		if (ParameterInterface)
-		{
-			ParameterInterface->Shutdown();
-			ParameterInterface = nullptr;
-		}
+		ResetParameters();
 	}
 
 	// Set this to immediately be inactive
@@ -1089,234 +1071,6 @@ void UAudioComponent::Deactivate()
 	}
 }
 
-void UAudioComponent::SetFloatParameter( const FName InName, const float InFloat )
-{
-	if (InName != NAME_None)
-	{
-		bool bFound = false;
-
-		// First see if an entry for this name already exists
-		for (FAudioComponentParam& P : InstanceParameters)
-		{
-			if (P.ParamName == InName)
-			{
-				P.FloatParam = InFloat;
-				bFound = true;
-				break;
-			}
-		}
-
-		// We didn't find one, so create a new one.
-		if (!bFound)
-		{
-			const int32 NewParamIndex = InstanceParameters.AddDefaulted();
-			InstanceParameters[ NewParamIndex ].ParamName = InName;
-			InstanceParameters[ NewParamIndex ].FloatParam = InFloat;
-		}
-
-		// If we're active we need to push this value to the ActiveSound
-		if (IsActive())
-		{
-			if (FAudioDevice* AudioDevice = GetAudioDevice())
-			{
-				DECLARE_CYCLE_STAT(TEXT("FAudioThreadTask.SetFloatParameter"), STAT_AudioSetFloatParameter, STATGROUP_AudioThreadCommands);
-
-				const uint64 MyAudioComponentID = AudioComponentID;
-				FAudioThread::RunCommandOnAudioThread([AudioDevice, MyAudioComponentID, InName, InFloat]()
-				{
-					FActiveSound* ActiveSound = AudioDevice->FindActiveSound(MyAudioComponentID);
-					if (ActiveSound)
-					{
-						ActiveSound->SetFloatParameter(InName, InFloat);
-					}
-				}, GET_STATID(STAT_AudioSetFloatParameter));
-			}
-		}
-	}
-}
-
-TScriptInterface<IAudioParameterInterface> UAudioComponent::GetParameterInterface() const
-{
-	return ParameterInterface;
-}
-
-void UAudioComponent::SetWaveParameter( const FName InName, USoundWave* InWave )
-{
-	if (InName != NAME_None)
-	{
-		bool bFound = false;
-		// First see if an entry for this name already exists
-		for (FAudioComponentParam& P : InstanceParameters)
-		{
-			if (P.ParamName == InName)
-			{
-				P.SoundWaveParam = InWave;
-				bFound = true;
-				break;
-			}
-		}
-
-		// We didn't find one, so create a new one.
-		if (!bFound)
-		{
-			const int32 NewParamIndex = InstanceParameters.AddDefaulted();
-			InstanceParameters[NewParamIndex].ParamName = InName;
-			InstanceParameters[NewParamIndex].SoundWaveParam = InWave;
-		}
-
-		// If we're active we need to push this value to the ActiveSound
-		if (IsActive())
-		{
-			if (FAudioDevice* AudioDevice = GetAudioDevice())
-			{
-				DECLARE_CYCLE_STAT(TEXT("FAudioThreadTask.SetWaveParameter"), STAT_AudioSetWaveParameter, STATGROUP_AudioThreadCommands);
-
-				const uint64 MyAudioComponentID = AudioComponentID;
-				FAudioThread::RunCommandOnAudioThread([AudioDevice, MyAudioComponentID, InName, InWave]()
-				{
-					FActiveSound* ActiveSound = AudioDevice->FindActiveSound(MyAudioComponentID);
-					if (ActiveSound)
-					{
-						ActiveSound->SetWaveParameter(InName, InWave);
-					}
-				}, GET_STATID(STAT_AudioSetWaveParameter));
-			}
-		}
-	}
-}
-
-void UAudioComponent::SetBoolParameter( const FName InName, const bool InBool )
-{
-	if (InName != NAME_None)
-	{
-		bool bFound = false;
-		// First see if an entry for this name already exists
-		for (FAudioComponentParam& P : InstanceParameters)
-		{
-			if (P.ParamName == InName)
-			{
-				P.BoolParam = InBool;
-				bFound = true;
-				break;
-			}
-		}
-
-		// We didn't find one, so create a new one.
-		if (!bFound)
-		{
-			const int32 NewParamIndex = InstanceParameters.AddDefaulted();
-			InstanceParameters[ NewParamIndex ].ParamName = InName;
-			InstanceParameters[ NewParamIndex ].BoolParam = InBool;
-		}
-
-		// If we're active we need to push this value to the ActiveSound
-		if (IsActive())
-		{
-			if (FAudioDevice* AudioDevice = GetAudioDevice())
-			{
-				DECLARE_CYCLE_STAT(TEXT("FAudioThreadTask.SetBoolParameter"), STAT_AudioSetBoolParameter, STATGROUP_AudioThreadCommands);
-
-				const uint64 MyAudioComponentID = AudioComponentID;
-				FAudioThread::RunCommandOnAudioThread([AudioDevice, MyAudioComponentID, InName, InBool]()
-				{
-					FActiveSound* ActiveSound = AudioDevice->FindActiveSound(MyAudioComponentID);
-					if (ActiveSound)
-					{
-						ActiveSound->SetBoolParameter(InName, InBool);
-					}
-				}, GET_STATID(STAT_AudioSetBoolParameter));
-			}
-		}
-	}
-}
-
-void UAudioComponent::SetIntParameter( const FName InName, const int32 InInt )
-{
-	if (InName != NAME_None)
-	{
-		bool bFound = false;
-		// First see if an entry for this name already exists
-		for (FAudioComponentParam& P : InstanceParameters)
-		{
-			if (P.ParamName == InName)
-			{
-				P.IntParam = InInt;
-				bFound = true;
-				break;
-			}
-		}
-
-		// We didn't find one, so create a new one.
-		if (!bFound)
-		{
-			const int32 NewParamIndex = InstanceParameters.AddDefaulted();
-			InstanceParameters[NewParamIndex].ParamName = InName;
-			InstanceParameters[NewParamIndex].IntParam = InInt;
-		}
-
-		// If we're active we need to push this value to the ActiveSound
-		if (IsActive())
-		{
-			if (FAudioDevice* AudioDevice = GetAudioDevice())
-			{
-				DECLARE_CYCLE_STAT(TEXT("FAudioThreadTask.SetIntParameter"), STAT_AudioSetIntParameter, STATGROUP_AudioThreadCommands);
-
-				const uint64 MyAudioComponentID = AudioComponentID;
-				FAudioThread::RunCommandOnAudioThread([AudioDevice, MyAudioComponentID, InName, InInt]()
-				{
-					FActiveSound* ActiveSound = AudioDevice->FindActiveSound(MyAudioComponentID);
-					if (ActiveSound)
-					{
-						ActiveSound->SetIntParameter(InName, InInt);
-					}
-				}, GET_STATID(STAT_AudioSetIntParameter));
-			}
-		}
-	}
-}
-
-void UAudioComponent::SetSoundParameter(const FAudioComponentParam& Param)
-{
-	if (Param.ParamName != NAME_None)
-	{
-		bool bFound = false;
-		// First see if an entry for this name already exists
-		for (FAudioComponentParam& P : InstanceParameters)
-		{
-			if (P.ParamName == Param.ParamName)
-			{
-				P = Param;
-				bFound = true;
-				break;
-			}
-		}
-
-		// We didn't find one, so create a new one.
-		if (!bFound)
-		{
-			const int32 NewParamIndex = InstanceParameters.Add(Param);
-		}
-
-		if (IsActive())
-		{
-			if (FAudioDevice* AudioDevice = GetAudioDevice())
-			{
-				DECLARE_CYCLE_STAT(TEXT("FAudioThreadTask.SetSoundParameter"), STAT_AudioSetSoundParameter, STATGROUP_AudioThreadCommands);
-
-				const uint64 MyAudioComponentID = AudioComponentID;
-				FAudioThread::RunCommandOnAudioThread([AudioDevice, MyAudioComponentID, Param]()
-				{
-					FActiveSound* ActiveSound = AudioDevice->FindActiveSound(MyAudioComponentID);
-					if (ActiveSound)
-					{
-						ActiveSound->SetSoundParameter(Param);
-					}
-				}, GET_STATID(STAT_AudioSetSoundParameter));
-			}
-		}
-	}
-}
-
 void UAudioComponent::SetFadeInComplete()
 {
 	EAudioComponentPlayState PlayState = GetPlayState();
@@ -1344,6 +1098,11 @@ void UAudioComponent::SetIsVirtualized(bool bInIsVirtualized)
 	bIsVirtualized = bInIsVirtualized ? 1 : 0;
 }
 
+void UAudioComponent::SetWaveParameter(FName InName, USoundWave* InWave)
+{
+	SetObjectParameter(InName, Cast<UObject>(InWave));
+}
+
 void UAudioComponent::SetVolumeMultiplier(const float NewVolumeMultiplier)
 {
 	VolumeMultiplier = NewVolumeMultiplier;
diff --git a/Engine/Source/Runtime/Engine/Private/GameplayStatics.cpp b/Engine/Source/Runtime/Engine/Private/GameplayStatics.cpp
--- a/Engine/Source/Runtime/Engine/Private/GameplayStatics.cpp
+++ b/Engine/Source/Runtime/Engine/Private/GameplayStatics.cpp
@@ -1582,7 +1582,7 @@ void UGameplayStatics::PlaySoundAtLocation(const UObject* WorldContextObject, cl
 
 	if (FAudioDeviceHandle AudioDevice = ThisWorld->GetAudioDevice())
 	{
-		const TArray<FAudioComponentParam>* Params = InitialParams ? &InitialParams->AudioParams : nullptr;
+		const TArray<FAudioParameter>* Params = InitialParams ? &InitialParams->AudioParams : nullptr;
 		AudioDevice->PlaySoundAtLocation(Sound, ThisWorld, VolumeMultiplier, PitchMultiplier, StartTime, Location, Rotation, AttenuationSettings, ConcurrencySettings, Params, OwningActor);
 	}
 }
diff --git a/Engine/Source/Runtime/Engine/Private/SoundBase.cpp b/Engine/Source/Runtime/Engine/Private/SoundBase.cpp
--- a/Engine/Source/Runtime/Engine/Private/SoundBase.cpp
+++ b/Engine/Source/Runtime/Engine/Private/SoundBase.cpp
@@ -3,7 +3,7 @@
 
 #include "AudioDevice.h"
 #include "EngineDefines.h"
-#include "IAudioExtensionPlugin.h"
+#include "IAudioParameterTransmitter.h"
 #include "Sound/AudioSettings.h"
 #include "Sound/SoundClass.h"
 #include "Sound/SoundSubmix.h"
@@ -277,4 +277,209 @@ void USoundBase::RemoveUserDataOfClass(TSubclassOf<UAssetUserData> InUserDataCla
 const TArray<UAssetUserData*>* USoundBase::GetAssetUserDataArray() const
 {
 	return &ToRawPtrTArrayUnsafe(AssetUserData);
-}
\ No newline at end of file
+}
+
+TUniquePtr<Audio::IParameterTransmitter> USoundBase::CreateParameterTransmitter(Audio::FParameterTransmitterInitParams&& InParams) const
+{
+	class FDefaultParameterTransmitter : public Audio::IParameterTransmitter
+	{
+		uint64 InstanceID = INDEX_NONE;
+		TArray<FAudioParameter> AudioParameters;
+
+	public:
+		FDefaultParameterTransmitter(Audio::FParameterTransmitterInitParams&& InParams)
+			: InstanceID(MoveTemp(InParams.InstanceID))
+			, AudioParameters(MoveTemp(InParams.DefaultParams))
+		{
+		}
+
+		virtual ~FDefaultParameterTransmitter() = default;
+
+		bool Reset() override
+		{
+			AudioParameters.Reset();
+			return true;
+		}
+
+		bool SetParameter(FAudioParameter&& InValue) override
+		{
+			if (FAudioParameter* Param = FAudioParameter::FindOrAddParam(AudioParameters, InValue.ParamName))
+			{
+				// Cannot do wholesale move because legacy system (i.e. SoundCues) support
+				// multiple value types mapped to a single parameter FName
+				switch (InValue.ParamType)
+				{
+					case EAudioParameterType::Boolean:
+					{
+						Param->BoolParam = InValue.BoolParam;
+					}
+					break;
+
+					case EAudioParameterType::BooleanArray:
+					{
+						Param->ArrayBoolParam = MoveTemp(InValue.ArrayBoolParam);
+					}
+					break;
+
+					case EAudioParameterType::Float:
+					{
+						Param->FloatParam = InValue.FloatParam;
+					}
+					break;
+
+					case EAudioParameterType::FloatArray:
+					{
+						Param->ArrayFloatParam = MoveTemp(InValue.ArrayFloatParam);
+					}
+					break;
+
+					case EAudioParameterType::Integer:
+					case EAudioParameterType::NoneArray:
+					{
+						Param->IntParam = InValue.IntParam;
+					}
+					break;
+
+					case EAudioParameterType::IntegerArray:
+					{
+						Param->ArrayIntParam = MoveTemp(InValue.ArrayIntParam);
+					}
+					break;
+
+					case EAudioParameterType::None:
+					{
+						*Param = MoveTemp(InValue);
+					}
+					break;
+
+					case EAudioParameterType::Object:
+					{
+						Param->ObjectParam = MoveTemp(InValue.ObjectParam);
+						Param->ObjectProxies = MoveTemp(InValue.ObjectProxies);
+					}
+					break;
+
+					case EAudioParameterType::ObjectArray:
+					{
+						Param->ArrayObjectParam = MoveTemp(InValue.ArrayObjectParam);
+						Param->ObjectProxies = MoveTemp(InValue.ObjectProxies);
+					}
+					break;
+
+					case EAudioParameterType::String:
+					{
+						Param->StringParam = MoveTemp(InValue.StringParam);
+					}
+					break;
+
+					case EAudioParameterType::StringArray:
+					{
+						Param->ArrayStringParam = MoveTemp(InValue.ArrayStringParam);
+					}
+					break;
+
+					default:
+					{
+						checkNoEntry();
+					}
+				}
+
+				InValue = FAudioParameter();
+
+				return true;
+			}
+
+			InValue = FAudioParameter();
+			return false;
+		}
+
+		uint64 GetInstanceID() const override
+		{
+			return InstanceID;
+		}
+
+		bool GetParameter(FName InName, FAudioParameter& OutValue) const override
+		{
+			if (const FAudioParameter* Param = FAudioParameter::FindParam(AudioParameters, InName))
+			{
+				OutValue = *Param;
+				return true;
+			}
+
+			return false;
+		}
+
+		TArray<UObject*> GetReferencedObjects() const override
+		{
+			TArray<UObject*> Objects;
+			for (const FAudioParameter& Param : AudioParameters)
+			{
+				if (Param.ObjectParam)
+				{
+					Objects.Add(Param.ObjectParam);
+				}
+
+				for (UObject* Object : Param.ArrayObjectParam)
+				{
+					if (Object)
+					{
+						Objects.Add(Object);
+					}
+				}
+			}
+
+			return Objects;
+		}
+
+		TUniquePtr<Audio::IParameterTransmitter> Clone() const override
+		{
+			return MakeUnique<FDefaultParameterTransmitter>(*this);
+		}
+	};
+
+	return MakeUnique<FDefaultParameterTransmitter>(MoveTemp(InParams));
+}
+
+void USoundBase::InitParameters(TArray<FAudioParameter>& InParametersToInit, FName InFeatureName)
+{
+	for (int32 i = InParametersToInit.Num() - 1; i >= 0; --i)
+	{
+		if (!IsParameterValid(InParametersToInit[i]))
+		{
+			InParametersToInit.RemoveAtSwap(i, 1, false /* bAllowShrinking */);
+		}
+	}
+}
+
+bool USoundBase::IsParameterValid(const FAudioParameter& InParameter) const
+{
+	switch (InParameter.ParamType)
+	{
+		case EAudioParameterType::BooleanArray:
+		case EAudioParameterType::FloatArray:
+		case EAudioParameterType::IntegerArray:
+		case EAudioParameterType::NoneArray:
+		case EAudioParameterType::ObjectArray:
+		case EAudioParameterType::StringArray:
+		case EAudioParameterType::String:
+		{
+			return false;
+		}
+
+		case EAudioParameterType::Object:
+		case EAudioParameterType::None:
+		{
+			if (InParameter.ObjectParam)
+			{
+				return InParameter.ObjectParam->GetClass()->IsChildOf(USoundWave::StaticClass());
+			}
+		}
+
+		default:
+		{
+			break;
+		}
+	}
+
+	return true;
+}
diff --git a/Engine/Source/Runtime/Engine/Private/SoundNodeBranch.cpp b/Engine/Source/Runtime/Engine/Private/SoundNodeBranch.cpp
--- a/Engine/Source/Runtime/Engine/Private/SoundNodeBranch.cpp
+++ b/Engine/Source/Runtime/Engine/Private/SoundNodeBranch.cpp
@@ -14,17 +14,16 @@ USoundNodeBranch::USoundNodeBranch(const FObjectInitializer& ObjectInitializer)
 {
 }
 
-void USoundNodeBranch::ParseNodes( FAudioDevice* AudioDevice, const UPTRINT NodeWaveInstanceHash, FActiveSound& ActiveSound, const FSoundParseParameters& ParseParams, TArray<FWaveInstance*>& WaveInstances )
+void USoundNodeBranch::ParseNodes(FAudioDevice* AudioDevice, const UPTRINT NodeWaveInstanceHash, FActiveSound& ActiveSound, const FSoundParseParameters& ParseParams, TArray<FWaveInstance*>& WaveInstances)
 {
 	BranchPurpose BranchToUse = BranchPurpose::ParameterUnset;
-	bool bParamValue = false;
 
-	if (ActiveSound.GetBoolParameter( BoolParameterName, bParamValue ))
+	FAudioParameter ParamValue;
+	if (ActiveSound.GetTransmitter()->GetParameter(BoolParameterName, ParamValue))
 	{
-		BranchToUse = (bParamValue ? BranchPurpose::ParameterTrue : BranchPurpose::ParameterFalse);
+		BranchToUse = (ParamValue.BoolParam ? BranchPurpose::ParameterTrue : BranchPurpose::ParameterFalse);
 	}
 
-
 	const int32 ChildNodeIndex = (int32)BranchToUse;
 	if (ChildNodeIndex < ChildNodes.Num() && ChildNodes[ChildNodeIndex])
 	{
diff --git a/Engine/Source/Runtime/Engine/Private/SoundNodeModulatorContinuous.cpp b/Engine/Source/Runtime/Engine/Private/SoundNodeModulatorContinuous.cpp
--- a/Engine/Source/Runtime/Engine/Private/SoundNodeModulatorContinuous.cpp
+++ b/Engine/Source/Runtime/Engine/Private/SoundNodeModulatorContinuous.cpp
@@ -6,24 +6,23 @@
 
 float FModulatorContinuousParams::GetValue(const FActiveSound& ActiveSound) const
 {
-	float ParamFloat = 0.f;
-
-	if (!ActiveSound.GetFloatParameter(ParameterName, ParamFloat))
+	FAudioParameter Param;
+	if (!ActiveSound.GetTransmitter()->GetParameter(ParameterName, Param))
 	{
-		ParamFloat = Default;
+		Param.FloatParam = Default;
 	}
 
 	if(ParamMode == MPM_Direct)
 	{
-		return ParamFloat;
+		return Param.FloatParam;
 	}
 	else if(ParamMode == MPM_Abs)
 	{
-		ParamFloat = FMath::Abs(ParamFloat);
+		Param.FloatParam = FMath::Abs(Param.FloatParam);
 	}
 
 	float Gradient;
-	if(MaxInput <= MinInput)
+	if (MaxInput <= MinInput)
 	{
 		Gradient = 0.f;
 	}
@@ -32,7 +31,7 @@ float FModulatorContinuousParams::GetValue(const FActiveSound& ActiveSound) cons
 		Gradient = (MaxOutput - MinOutput)/(MaxInput - MinInput);
 	}
 
-	const float ClampedParam = FMath::Clamp(ParamFloat, MinInput, MaxInput);
+	const float ClampedParam = FMath::Clamp(Param.FloatParam, MinInput, MaxInput);
 
 	return MinOutput + ((ClampedParam - MinInput) * Gradient);
 }
diff --git a/Engine/Source/Runtime/Engine/Private/SoundNodeParamCrossFade.cpp b/Engine/Source/Runtime/Engine/Private/SoundNodeParamCrossFade.cpp
--- a/Engine/Source/Runtime/Engine/Private/SoundNodeParamCrossFade.cpp
+++ b/Engine/Source/Runtime/Engine/Private/SoundNodeParamCrossFade.cpp
@@ -13,10 +13,9 @@ USoundNodeParamCrossFade::USoundNodeParamCrossFade(const FObjectInitializer& Obj
 
 float USoundNodeParamCrossFade::GetCurrentDistance(FAudioDevice* AudioDevice, FActiveSound& ActiveSound, const FSoundParseParameters& ParseParams) const
 {
-	float ParamValue = 0.0f;
-	
-	ActiveSound.GetFloatParameter(ParamName, ParamValue);
-	return ParamValue;
+	FAudioParameter ParamValue;
+	ActiveSound.GetTransmitter()->GetParameter(ParamName, ParamValue);
+	return ParamValue.FloatParam;
 }
 
 bool USoundNodeParamCrossFade::AllowCrossfading(FActiveSound& ActiveSound) const
diff --git a/Engine/Source/Runtime/Engine/Private/SoundNodeSwitch.cpp b/Engine/Source/Runtime/Engine/Private/SoundNodeSwitch.cpp
--- a/Engine/Source/Runtime/Engine/Private/SoundNodeSwitch.cpp
+++ b/Engine/Source/Runtime/Engine/Private/SoundNodeSwitch.cpp
@@ -12,32 +12,31 @@ USoundNodeSwitch::USoundNodeSwitch(const FObjectInitializer& ObjectInitializer)
 {
 }
 
-void USoundNodeSwitch::ParseNodes( FAudioDevice* AudioDevice, const UPTRINT NodeWaveInstanceHash, FActiveSound& ActiveSound, const FSoundParseParameters& ParseParams, TArray<FWaveInstance*>& WaveInstances )
+void USoundNodeSwitch::ParseNodes(FAudioDevice* AudioDevice, const UPTRINT NodeWaveInstanceHash, FActiveSound& ActiveSound, const FSoundParseParameters& ParseParams, TArray<FWaveInstance*>& WaveInstances)
 {
-	int32 ChildNodeIndex = 0;
-
-	if (ActiveSound.GetIntParameter( IntParameterName, ChildNodeIndex ))
+	FAudioParameter Param;
+	if (ActiveSound.GetTransmitter()->GetParameter(IntParameterName, Param))
 	{
-		ChildNodeIndex += 1;
+		Param.IntParam += 1;
 	}
 	
-	if (ChildNodeIndex < 0 || ChildNodeIndex >= ChildNodes.Num())
+	if (Param.IntParam < 0 || Param.IntParam >= ChildNodes.Num())
 	{
-		ChildNodeIndex = 0;
+		Param.IntParam = 0;
 	}
 
-	if (ChildNodeIndex < ChildNodes.Num() && ChildNodes[ChildNodeIndex])
+	if (Param.IntParam < ChildNodes.Num() && ChildNodes[Param.IntParam])
 	{
-		ChildNodes[ChildNodeIndex]->ParseNodes(AudioDevice, GetNodeWaveInstanceHash(NodeWaveInstanceHash, ChildNodes[ChildNodeIndex], ChildNodeIndex), ActiveSound, ParseParams, WaveInstances);
+		ChildNodes[Param.IntParam]->ParseNodes(AudioDevice, GetNodeWaveInstanceHash(NodeWaveInstanceHash, ChildNodes[Param.IntParam], Param.IntParam), ActiveSound, ParseParams, WaveInstances);
 	}
 }
 
 void USoundNodeSwitch::CreateStartingConnectors()
 {
-	InsertChildNode( ChildNodes.Num() );
-	InsertChildNode( ChildNodes.Num() );
-	InsertChildNode( ChildNodes.Num() );
-	InsertChildNode( ChildNodes.Num() );
+	InsertChildNode(ChildNodes.Num());
+	InsertChildNode(ChildNodes.Num());
+	InsertChildNode(ChildNodes.Num());
+	InsertChildNode(ChildNodes.Num());
 }
 
 #if WITH_EDITOR
diff --git a/Engine/Source/Runtime/Engine/Private/SoundNodeWaveParam.cpp b/Engine/Source/Runtime/Engine/Private/SoundNodeWaveParam.cpp
--- a/Engine/Source/Runtime/Engine/Private/SoundNodeWaveParam.cpp
+++ b/Engine/Source/Runtime/Engine/Private/SoundNodeWaveParam.cpp
@@ -20,18 +20,19 @@ float USoundNodeWaveParam::GetDuration()
 	return INDEFINITELY_LOOPING_DURATION;
 }
 
-void USoundNodeWaveParam::ParseNodes( FAudioDevice* AudioDevice, const UPTRINT NodeWaveInstanceHash, FActiveSound& ActiveSound, const FSoundParseParameters& ParseParams, TArray<FWaveInstance*>& WaveInstances )
+void USoundNodeWaveParam::ParseNodes(FAudioDevice* AudioDevice, const UPTRINT NodeWaveInstanceHash, FActiveSound& ActiveSound, const FSoundParseParameters& ParseParams, TArray<FWaveInstance*>& WaveInstances)
 {
-	USoundWave* NewWave = NULL;
-	ActiveSound.GetWaveParameter( WaveParameterName, NewWave );
-	if( NewWave != NULL )
+	FAudioParameter ParamValue;
+	ActiveSound.GetTransmitter()->GetParameter(WaveParameterName, ParamValue);
+
+	if (USoundWave* NewWave = Cast<USoundWave>(ParamValue.ObjectParam))
 	{
-		NewWave->Parse( AudioDevice, GetNodeWaveInstanceHash(NodeWaveInstanceHash, (UPTRINT)NewWave, 0), ActiveSound, ParseParams, WaveInstances );
+		NewWave->Parse(AudioDevice, GetNodeWaveInstanceHash(NodeWaveInstanceHash, (UPTRINT)NewWave, 0), ActiveSound, ParseParams, WaveInstances);
 	}
 	else
 	{
 		// use the default node linked to us, if any
-		Super::ParseNodes( AudioDevice, NodeWaveInstanceHash, ActiveSound, ParseParams, WaveInstances );
+		Super::ParseNodes(AudioDevice, NodeWaveInstanceHash, ActiveSound, ParseParams, WaveInstances);
 	}
 }
 
diff --git a/Engine/Source/Runtime/Engine/Public/ActiveSound.h b/Engine/Source/Runtime/Engine/Public/ActiveSound.h
--- a/Engine/Source/Runtime/Engine/Public/ActiveSound.h
+++ b/Engine/Source/Runtime/Engine/Public/ActiveSound.h
@@ -362,19 +362,19 @@ private:
 
 	TMap<UPTRINT, FWaveInstance*> WaveInstances;
 
-	class FInstanceTransmitter
+	class FParameterTransmitter
 	{
 	public:
-		FInstanceTransmitter() = default;
-		FInstanceTransmitter(FInstanceTransmitter&&) = default;
-		FInstanceTransmitter& operator=(FInstanceTransmitter&&) = default;
+		FParameterTransmitter() = default;
+		FParameterTransmitter(FParameterTransmitter&&) = default;
+		FParameterTransmitter& operator=(FParameterTransmitter&&) = default;
 
-		FInstanceTransmitter(TUniquePtr<IAudioInstanceTransmitter>&& InTransmitterImpl)
+		FParameterTransmitter(TUniquePtr<Audio::IParameterTransmitter>&& InTransmitterImpl)
 		:	TransmitterImpl(MoveTemp(InTransmitterImpl))
 		{
 		}
 
-		FInstanceTransmitter(const FInstanceTransmitter& InOther)
+		FParameterTransmitter(const FParameterTransmitter& InOther)
 		{
 			if (InOther.IsValid())
 			{
@@ -382,7 +382,7 @@ private:
 			}
 		}
 
-		FInstanceTransmitter& operator=(const FInstanceTransmitter& InOther)
+		FParameterTransmitter& operator=(const FParameterTransmitter& InOther)
 		{
 			TransmitterImpl.Reset();
 			if (InOther.IsValid())
@@ -392,7 +392,7 @@ private:
 			return *this;
 		}
 
-		FInstanceTransmitter& operator=(TUniquePtr<IAudioInstanceTransmitter>&& InTransmitterImpl)
+		FParameterTransmitter& operator=(TUniquePtr<Audio::IParameterTransmitter>&& InTransmitterImpl)
 		{
 			TransmitterImpl = MoveTemp(InTransmitterImpl);
 			return *this;
@@ -403,37 +403,38 @@ private:
 			return TransmitterImpl.IsValid();
 		}
 
-		IAudioInstanceTransmitter* Get()
+		Audio::IParameterTransmitter* Get()
 		{
 			return TransmitterImpl.Get();
 		}
-		const IAudioInstanceTransmitter* Get() const
+		const Audio::IParameterTransmitter* Get() const
 		{
 			return TransmitterImpl.Get();
 		}
 
-		IAudioInstanceTransmitter* operator->()
+		Audio::IParameterTransmitter* operator->()
 		{
 			return TransmitterImpl.Get();
 		}
 
-		const IAudioInstanceTransmitter* operator->() const
+		const Audio::IParameterTransmitter* operator->() const
 		{
 			return TransmitterImpl.Get();
 		}
 
 	private:
-		TUniquePtr<IAudioInstanceTransmitter> TransmitterImpl;
+		TUniquePtr<Audio::IParameterTransmitter> TransmitterImpl;
 	};
 
-	FInstanceTransmitter InstanceTransmitter;
+	FParameterTransmitter InstanceTransmitter;
 
 public:
-	IAudioInstanceTransmitter* GetTransmitter() 
+	Audio::IParameterTransmitter* GetTransmitter()
 	{
 		return InstanceTransmitter.Get();
 	}
-	const IAudioInstanceTransmitter* GetTransmitter() const
+
+	const Audio::IParameterTransmitter* GetTransmitter() const
 	{
 		return InstanceTransmitter.Get();
 	}
@@ -641,13 +642,11 @@ public:
 	TMap<UPTRINT,uint32> SoundNodeOffsetMap;
 	TArray<uint8> SoundNodeData;
 
-	TArray<FAudioComponentParam> InstanceParameters;
-
 	// Whether or not there are Source Bus Sends that have not been sent to the render thread
 	bool bHasNewBusSends;
 
 	// Bus send(s) that have not yet been sent to the render thread
-	TArray<TTuple<EBusSendType, FSoundSourceBusSendInfo>> newBusSends;
+	TArray<TTuple<EBusSendType, FSoundSourceBusSendInfo>> NewBusSends;
 
 	FSoundModulationDefaultRoutingSettings ModulationRouting;
 
@@ -683,45 +682,6 @@ public:
 	/** Gets total concurrency gain stage based on all concurrency memberships of sound */
 	float GetTotalConcurrencyVolumeScale() const;
 
-	/** Sets a float instance parameter for the ActiveSound */
-	void SetFloatParameter(const FName InName, const float InFloat);
-
-	/** Sets a wave instance parameter for the ActiveSound */
-	void SetWaveParameter(const FName InName, class USoundWave* InWave);
-
-	/** Sets a boolean instance parameter for the ActiveSound */
-	void SetBoolParameter(const FName InName, const bool InBool);
-
-	/** Sets an integer instance parameter for the ActiveSound */
-	void SetIntParameter(const FName InName, const int32 InInt);
-
-	/** Sets the audio component parameter on the active sound. Note: this can be set without audio components if they are set when active sound is created. */
-	void SetSoundParameter(const FAudioComponentParam& Param);
-
-	/**
-	 * Try and find an Instance Parameter with the given name and if we find it return the float value.
-	 * @return true if float for parameter was found, otherwise false
-	 */
-	bool GetFloatParameter(const FName InName, float& OutFloat) const;
-
-	/**
-	 *Try and find an Instance Parameter with the given name and if we find it return the USoundWave value.
-	 * @return true if USoundWave for parameter was found, otherwise false
-	 */
-	bool GetWaveParameter(const FName InName, USoundWave*& OutWave) const;
-
-	/**
-	 *Try and find an Instance Parameter with the given name and if we find it return the boolean value.
-	 * @return true if boolean for parameter was found, otherwise false
-	 */
-	bool GetBoolParameter(const FName InName, bool& OutBool) const;
-
-	/**
-	 *Try and find an Instance Parameter with the given name and if we find it return the integer value.
-	 * @return true if boolean for parameter was found, otherwise false
-	 */
-	bool GetIntParameter(const FName InName, int32& OutInt) const;
-
 	void CollectAttenuationShapesForVisualization(TMultiMap<EAttenuationShape::Type, FBaseAttenuationSettings::AttenuationShapeDetails>& ShapeDetailsMap) const;
 
 	/**
diff --git a/Engine/Source/Runtime/Engine/Public/AudioDevice.h b/Engine/Source/Runtime/Engine/Public/AudioDevice.h
--- a/Engine/Source/Runtime/Engine/Public/AudioDevice.h
+++ b/Engine/Source/Runtime/Engine/Public/AudioDevice.h
@@ -758,12 +758,17 @@ public:
 	 * @param	USoundConcurrency	The sound's sound concurrency settings to use (optional). Will use the USoundBase's USoundConcurrency if not specified.
 	 * @param	Params				An optional list of audio component params to immediately apply to a sound.
 	 */
-	void PlaySoundAtLocation(USoundBase* Sound, UWorld* World, float VolumeMultiplier, float PitchMultiplier, float StartTime, const FVector& Location, const FRotator& Rotation, USoundAttenuation* AttenuationSettings = nullptr, USoundConcurrency* ConcurrencySettings = nullptr, const TArray<FAudioComponentParam>* Params = nullptr, AActor* OwningActor = nullptr);
+	void PlaySoundAtLocation(USoundBase* Sound, UWorld* World, float VolumeMultiplier, float PitchMultiplier, float StartTime, const FVector& Location, const FRotator& Rotation, USoundAttenuation* AttenuationSettings = nullptr, USoundConcurrency* ConcurrencySettings = nullptr, const TArray<FAudioParameter>* Params = nullptr, AActor* OwningActor = nullptr);
 
 	/**
 	 * Adds an active sound to the audio device
 	 */
-	void AddNewActiveSound(const FActiveSound& ActiveSound);
+	void AddNewActiveSound(const FActiveSound& ActiveSound, const TArray<FAudioParameter>* InDefaultParams = nullptr);
+
+	/**
+	 * Adds an active sound to the audio device
+	 */
+	void AddNewActiveSound(const FActiveSound& ActiveSound, TArray<FAudioParameter>&& InDefaultParams);
 
 	/**
 	 * Attempts to retrigger a provided loop
@@ -1451,7 +1456,7 @@ private:
 	 * Adds an active sound to the audio device. Can be a new active sound or one provided by the re-triggering
 	 * loop system.
 	 */
-	void AddNewActiveSoundInternal(const FActiveSound& ActiveSound, FAudioVirtualLoop* VirtualLoop);
+	void AddNewActiveSoundInternal(const FActiveSound& InActiveSound, TArray<FAudioParameter>&& InDefaultParams, FAudioVirtualLoop* InVirtualLoopToRetrigger = nullptr);
 
 	/**
 	 * Reports if a sound fails to start when attempting to create a new active sound.
