Missing file

diff --git a/Engine/Source/Runtime/AudioExtensions/Private/IAudioParameterTransmitter.cpp b/Engine/Source/Runtime/AudioExtensions/Private/IAudioParameterTransmitter.cpp
new file mode 100644
--- /dev/null
+++ b/Engine/Source/Runtime/AudioExtensions/Private/IAudioParameterTransmitter.cpp
@@ -0,0 +1,17 @@
+// Copyright Epic Games, Inc. All Rights Reserved.
+#include "IAudioParameterTransmitter.h"
+
+#include "AudioParameterInterface.h"
+
+namespace Audio
+{
+	bool ILegacyParameterTransmitter::GetParameter(FName InName, FAudioParameter& OutValue) const
+	{
+		return false;
+	}
+
+	TArray<UObject*> ILegacyParameterTransmitter::GetReferencedObjects() const
+	{
+		return { };
+	}
+} // namespace Audio
\ No newline at end of file
