Missing header

diff --git a/Engine/Source/Runtime/AudioExtensions/Public/IAudioParameterTransmitter.h b/Engine/Source/Runtime/AudioExtensions/Public/IAudioParameterTransmitter.h
new file mode 100644
--- /dev/null
+++ b/Engine/Source/Runtime/AudioExtensions/Public/IAudioParameterTransmitter.h
@@ -0,0 +1,57 @@
+// Copyright Epic Games, Inc. All Rights Reserved.
+#pragma once
+
+#include "AudioParameterInterface.h"
+#include "CoreMinimal.h"
+#include "IAudioProxyInitializer.h"
+#include "Templates/UniquePtr.h"
+
+
+namespace Audio
+{
+	/** Data passed to CreateParameterTransmitter. */
+	struct AUDIOEXTENSIONS_API FParameterTransmitterInitParams
+	{
+		// Unique ID for this audio instance.
+		uint64 InstanceID = INDEX_NONE;
+
+		// Audio sample rate.
+		float SampleRate = 0.0f;
+
+		TArray<FAudioParameter> DefaultParams;
+	};
+
+	// Parameter getter & reference collector for legacy parameter system.
+	// (i.e. backwards compatibility with the SoundCue system). None of this
+	// should be used by future systems (i.e. MetaSounds) as object references
+	// from parameters should NOT be cached on threads other than the GameThread.
+	class AUDIOEXTENSIONS_API ILegacyParameterTransmitter
+	{
+		public:
+			virtual ~ILegacyParameterTransmitter() = default;
+
+			virtual bool GetParameter(FName InName, FAudioParameter& OutParam) const;
+			virtual TArray<UObject*> GetReferencedObjects() const;
+	};
+
+	/** Interface for a audio instance transmitter.
+	 *
+	 * An audio instance transmitter ushers control parameters to a single audio object instance.
+	 */
+	class AUDIOEXTENSIONS_API IParameterTransmitter : public ILegacyParameterTransmitter
+	{
+		public:
+			virtual ~IParameterTransmitter() = default;
+
+			virtual bool Reset() = 0;
+
+			// Return the instance ID
+			virtual uint64 GetInstanceID() const = 0;
+
+			// Parameter Setters
+			virtual bool SetParameter(FAudioParameter&& InValue) = 0;
+
+			// Create a copy of the instance transmitter.
+			virtual TUniquePtr<Audio::IParameterTransmitter> Clone() const = 0;
+	};
+} // namespace Audio
\ No newline at end of file
