Move (invalid) IsValid checks to IsNone FName checks & make display names for parameter types consistent for BP readibility

diff --git a/Engine/Source/Runtime/AudioExtensions/Public/AudioParameterInterface.h b/Engine/Source/Runtime/AudioExtensions/Public/AudioParameterInterface.h
--- a/Engine/Source/Runtime/AudioExtensions/Public/AudioParameterInterface.h
+++ b/Engine/Source/Runtime/AudioExtensions/Public/AudioParameterInterface.h
@@ -291,27 +291,27 @@ struct AUDIOEXTENSIONS_API FAudioParameter
 	TObjectPtr<UObject> ObjectParam = nullptr;
 
 	// String value of parameter
-	UPROPERTY(EditAnywhere, BlueprintReadWrite, meta = (DisplayName = "Value", DisplayAfter = "ParamType", EditConditionHides, EditCondition = "ParamType == EAudioParameterType::String"), Category = AudioParameter)
+	UPROPERTY(EditAnywhere, BlueprintReadWrite, meta = (DisplayName = "Value (String)", DisplayAfter = "ParamType", EditConditionHides, EditCondition = "ParamType == EAudioParameterType::String"), Category = AudioParameter)
 	FString StringParam;
 
 	// Array Float value of parameter
-	UPROPERTY(EditAnywhere, BlueprintReadWrite, meta = (DisplayName = "Value", DisplayAfter = "ParamType", EditConditionHides, EditCondition = "ParamType == EAudioParameterType::FloatArray"), Category = AudioParameter)
+	UPROPERTY(EditAnywhere, BlueprintReadWrite, meta = (DisplayName = "Value (Float Array)", DisplayAfter = "ParamType", EditConditionHides, EditCondition = "ParamType == EAudioParameterType::FloatArray"), Category = AudioParameter)
 	TArray<float> ArrayFloatParam;
 
 	// Boolean value of parameter
-	UPROPERTY(EditAnywhere, BlueprintReadWrite, meta = (DisplayName = "Value", DisplayAfter = "ParamType", EditConditionHides, EditCondition = "ParamType == EAudioParameterType::BooleanArray"), Category = AudioParameter)
+	UPROPERTY(EditAnywhere, BlueprintReadWrite, meta = (DisplayName = "Value (Bool Array)", DisplayAfter = "ParamType", EditConditionHides, EditCondition = "ParamType == EAudioParameterType::BooleanArray"), Category = AudioParameter)
 	TArray<bool> ArrayBoolParam;
 
 	// Integer value of parameter
-	UPROPERTY(EditAnywhere, BlueprintReadWrite, meta = (DisplayName = "Value", DisplayAfter = "ParamType", EditConditionHides, EditCondition = "ParamType == EAudioParameterType::IntegerArray"), Category = AudioParameter)
+	UPROPERTY(EditAnywhere, BlueprintReadWrite, meta = (DisplayName = "Value (Int Array)", DisplayAfter = "ParamType", EditConditionHides, EditCondition = "ParamType == EAudioParameterType::IntegerArray"), Category = AudioParameter)
 	TArray<int32> ArrayIntParam;
 
 	// Object value of parameter
-	UPROPERTY(EditAnywhere, BlueprintReadWrite, meta = (DisplayName = "Value", DisplayAfter = "ParamType", EditConditionHides, EditCondition = "ParamType == EAudioParameterType::ObjectArray"), Category = AudioParameter)
+	UPROPERTY(EditAnywhere, BlueprintReadWrite, meta = (DisplayName = "Value (Object Array)", DisplayAfter = "ParamType", EditConditionHides, EditCondition = "ParamType == EAudioParameterType::ObjectArray"), Category = AudioParameter)
 	TArray<TObjectPtr<UObject>> ArrayObjectParam;
 
 	// String value of parameter
-	UPROPERTY(EditAnywhere, BlueprintReadWrite, meta = (DisplayName = "Value", DisplayAfter = "ParamType", EditConditionHides, EditCondition = "ParamType == EAudioParameterType::StringArray"), Category = AudioParameter)
+	UPROPERTY(EditAnywhere, BlueprintReadWrite, meta = (DisplayName = "Value (String Array)", DisplayAfter = "ParamType", EditConditionHides, EditCondition = "ParamType == EAudioParameterType::StringArray"), Category = AudioParameter)
 	TArray<FString> ArrayStringParam;
 
 	UPROPERTY(EditAnywhere, BlueprintReadWrite, meta = (DisplayName = "Type"), Category = AudioParameter)
@@ -323,7 +323,7 @@ struct AUDIOEXTENSIONS_API FAudioParameter
 	// Common find algorithm for default/legacy parameter system
 	static const FAudioParameter* FindParam(const TArray<FAudioParameter>& InParams, FName InParamName)
 	{
-		if (InParamName.IsValid())
+		if (!InParamName.IsNone())
 		{
 			for (const FAudioParameter& ExistingParam : InParams)
 			{
@@ -341,7 +341,7 @@ struct AUDIOEXTENSIONS_API FAudioParameter
 	static FAudioParameter* FindOrAddParam(TArray<FAudioParameter>& OutParams, FName InParamName)
 	{
 		FAudioParameter* Param = nullptr;
-		if (!InParamName.IsValid())
+		if (InParamName.IsNone())
 		{
 			return Param;
 		}
