Guard against potential race condition that can cause crash in debug builds on start-up when registering MetaSound Interfaces. General clean-up therein
#rnx

diff --git a/Engine/Plugins/Runtime/Metasound/Source/MetasoundEngine/Private/MetasoundInterface.cpp b/Engine/Plugins/Runtime/Metasound/Source/MetasoundEngine/Private/MetasoundInterface.cpp
--- a/Engine/Plugins/Runtime/Metasound/Source/MetasoundEngine/Private/MetasoundInterface.cpp
+++ b/Engine/Plugins/Runtime/Metasound/Source/MetasoundEngine/Private/MetasoundInterface.cpp
@@ -68,13 +68,15 @@ namespace Metasound
 				});
 
 				// For now, only support interfaces for UMetaSoundSource type
-				if (ensureAlways(UMetaSoundSource::StaticClass()->IsChildOf(Interface->Type)))
+				const UClass* SourceClass = UMetaSoundSource::StaticClass();
+				check(SourceClass);
+				if (ensureAlways(SourceClass->IsChildOf(Interface->Type)))
 				{
 					// For now, all external interfaces use the ParameterTransmitter and cannot be applied by default to new
 					// MetaSound objects. Could pipe either of these through in the future if required/desired.
 					constexpr bool bIsDefault = false;
 					RegisterInterface<UMetaSoundSource>(FrontendInterface, nullptr, bIsDefault, Audio::IParameterTransmitter::RouterName);
-					UE_LOG(LogMetaSound, Verbose, TEXT("Interface '%s' registered as MetaSoundSource assets."), *Interface->Name.ToString());
+					UE_LOG(LogMetaSound, Verbose, TEXT("Interface '%s' registered for asset type '%s'."), *Interface->Name.ToString(), *SourceClass->GetName());
 				}
 			};
 
diff --git a/Engine/Source/Runtime/AudioExtensions/Private/IAudioGeneratorInterfaceRegistry.cpp b/Engine/Source/Runtime/AudioExtensions/Private/IAudioGeneratorInterfaceRegistry.cpp
--- a/Engine/Source/Runtime/AudioExtensions/Private/IAudioGeneratorInterfaceRegistry.cpp
+++ b/Engine/Source/Runtime/AudioExtensions/Private/IAudioGeneratorInterfaceRegistry.cpp
@@ -13,9 +13,7 @@ namespace Audio
 		public:
 			virtual void IterateInterfaces(TFunction<void(FGeneratorInterfacePtr)> InFunction) const override
 			{
-				FGeneratorInterfaceQueryResults GeneratorQuery;
-				QueryInterfacesDelegate.Broadcast(GeneratorQuery);
-				for (const FGeneratorInterfacePtr& InterfacePtr : GeneratorQuery.GetInterfaces())
+				for (const FGeneratorInterfacePtr& InterfacePtr : Interfaces)
 				{
 					InFunction(InterfacePtr);
 				}
@@ -23,11 +21,11 @@ namespace Audio
 
 			virtual void RegisterInterface(FGeneratorInterfacePtr InInterface) override
 			{
-				RegistrationFunction(InInterface);
-				QueryInterfacesDelegate.AddLambda([Interface = InInterface](FGeneratorInterfaceQueryResults& Results)
+				Interfaces.Add(InInterface);
+				if (RegistrationFunction)
 				{
-					Results.AddInterface(Interface);
-				});
+					RegistrationFunction(InInterface);
+				}
 			}
 
 			virtual void OnRegistration(TUniqueFunction<void(FGeneratorInterfacePtr)>&& InFunction) override
diff --git a/Engine/Source/Runtime/AudioExtensions/Public/IAudioGeneratorInterfaceRegistry.h b/Engine/Source/Runtime/AudioExtensions/Public/IAudioGeneratorInterfaceRegistry.h
--- a/Engine/Source/Runtime/AudioExtensions/Public/IAudioGeneratorInterfaceRegistry.h
+++ b/Engine/Source/Runtime/AudioExtensions/Public/IAudioGeneratorInterfaceRegistry.h
@@ -2,7 +2,6 @@
 #pragma once
 
 #include "AudioParameterInterface.h"
-#include "Delegates/Delegate.h"
 #include "Templates/Function.h"
 #include "UObject/Class.h"
 
@@ -63,23 +62,6 @@ namespace Audio
 	};
 	using FGeneratorInterfacePtr = TSharedPtr<FGeneratorInterface>;
 
-	struct FGeneratorInterfaceQueryResults
-	{
-	private:
-		TArray<FGeneratorInterfacePtr> Interfaces;
-
-	public:
-		void AddInterface(FGeneratorInterfacePtr InInterfaceToAdd)
-		{
-			Interfaces.Add(InInterfaceToAdd);
-		}
-
-		const TArray<FGeneratorInterfacePtr>& GetInterfaces() const { return Interfaces; }
-	};
-
-	DECLARE_MULTICAST_DELEGATE_OneParam(FOnQueryAudioGeneratorInterface, FGeneratorInterfaceQueryResults&);
-	DECLARE_MULTICAST_DELEGATE_OneParam(FOnRegisterAudioGeneratorInterface, FGeneratorInterfacePtr);
-
 	class AUDIOEXTENSIONS_API IGeneratorInterfaceRegistry
 	{
 		static TUniquePtr<IGeneratorInterfaceRegistry> Instance;
@@ -98,8 +80,7 @@ namespace Audio
 		virtual void RegisterInterface(FGeneratorInterfacePtr InInterface) = 0;
 
 	protected:
-		mutable FOnQueryAudioGeneratorInterface QueryInterfacesDelegate;
-
+		TSet<FGeneratorInterfacePtr> Interfaces;
 		TUniqueFunction<void(FGeneratorInterfacePtr)> RegistrationFunction;
 	};
 } // namespace Audio
