adding IWaveformTransformation interface to AudioExtentions, and adding WaveformTransformations module to WaveformEditor.

diff --git a/Engine/Source/Runtime/AudioExtensions/Private/IWaveformTransformation.cpp b/Engine/Source/Runtime/AudioExtensions/Private/IWaveformTransformation.cpp
new file mode 100644
--- /dev/null
+++ b/Engine/Source/Runtime/AudioExtensions/Private/IWaveformTransformation.cpp
@@ -0,0 +1,20 @@
+// Copyright Epic Games, Inc. All Rights Reserved.
+
+
+#include "IWaveformTransformation.h"
+#include "Templates/SharedPointer.h"
+
+TArray<Audio::FTransformationPtr> UWaveformTransformationChain::CreateTransformations() const
+{
+	TArray<Audio::FTransformationPtr> TransformationPtrs;
+
+	for(UWaveformTransformationBase* Transformation : Transformations)
+	{
+		if(Transformation)
+		{
+			TransformationPtrs.Add(Transformation->CreateTransformation());
+		}
+	}
+	
+	return TransformationPtrs;
+}
diff --git a/Engine/Source/Runtime/AudioExtensions/Public/IWaveformTransformation.h b/Engine/Source/Runtime/AudioExtensions/Public/IWaveformTransformation.h
new file mode 100644
--- /dev/null
+++ b/Engine/Source/Runtime/AudioExtensions/Public/IWaveformTransformation.h
@@ -0,0 +1,65 @@
+// Copyright Epic Games, Inc. All Rights Reserved.
+
+#pragma once
+
+#include "SignalProcessing/Public/DSP/BufferVectorOperations.h"
+#include "UObject/Object.h"
+#include "Templates/SharedPointer.h"
+
+#include "IWaveformTransformation.generated.h"
+
+namespace Audio
+{
+	// information about the current state of the wave file we are transforming
+	struct AUDIOEXTENSIONS_API FWaveformTransformationWaveInfo
+	{
+		float SampleRate = 0.f;
+		int32 NumChannels = 0;
+		Audio::FAlignedFloatBuffer* Audio = nullptr;
+	};
+
+	/*
+	 * Base class for the object that processes waveform data
+	 * Pass tweakable variables from its paired settings UObject in the constructor in UWaveformTransformationBase::CreateTransformation
+	 *
+	 * note: WaveTransformation vs WaveformTransformation is to prevent UHT class name conflicts without having to namespace everything - remember this in derived classes!
+	 */
+	class AUDIOEXTENSIONS_API IWaveTransformation
+	{
+	public:
+
+		// Applies the transformation to the waveform and modifies WaveInfo with the resulting changes
+		virtual void ProcessAudio(FWaveformTransformationWaveInfo& InOutWaveInfo) const {};
+		
+		virtual bool SupportsRealtimePreview() const { return false; }
+		virtual bool CanChangeFileLength() const { return false; }
+		virtual bool CanChangeChannelCount() const { return false; }
+
+		virtual ~IWaveTransformation() {};
+	};
+
+	using FTransformationPtr = TUniquePtr<Audio::IWaveTransformation>;
+}
+
+// Base class to hold editor configurable properties for an arbitrary transformation of audio waveform data
+UCLASS(Abstract, EditInlineNew)
+class AUDIOEXTENSIONS_API UWaveformTransformationBase : public UObject
+{
+	GENERATED_BODY()
+
+public:
+	virtual Audio::FTransformationPtr CreateTransformation() const { return nullptr; }
+};
+
+// Object that holds an ordered list of transformations to perform on a sound wave
+UCLASS(EditInlineNew)
+class AUDIOEXTENSIONS_API UWaveformTransformationChain : public UObject
+{
+	GENERATED_BODY()
+	
+public:
+	UPROPERTY(EditAnywhere, Instanced, Category = "Transformations")
+	TArray<TObjectPtr<UWaveformTransformationBase>> Transformations;
+
+	TArray<Audio::FTransformationPtr> CreateTransformations() const;
+};
\ No newline at end of file
diff --git a/Engine/Source/Runtime/SignalProcessing/Private/FloatArrayMath.cpp b/Engine/Source/Runtime/SignalProcessing/Private/FloatArrayMath.cpp
--- a/Engine/Source/Runtime/SignalProcessing/Private/FloatArrayMath.cpp
+++ b/Engine/Source/Runtime/SignalProcessing/Private/FloatArrayMath.cpp
@@ -709,7 +709,7 @@ namespace Audio
 		}
 	}
 	
-	float ArrayMaxAbsValue(const TArrayView<float> InView)
+	float ArrayMaxAbsValue(const TArrayView<const float> InView)
 	{
 		const int32 Num = InView.Num();
 		const float* Data = InView.GetData();
diff --git a/Engine/Source/Runtime/SignalProcessing/Public/DSP/FloatArrayMath.h b/Engine/Source/Runtime/SignalProcessing/Public/DSP/FloatArrayMath.h
--- a/Engine/Source/Runtime/SignalProcessing/Public/DSP/FloatArrayMath.h
+++ b/Engine/Source/Runtime/SignalProcessing/Public/DSP/FloatArrayMath.h
@@ -113,7 +113,7 @@ namespace Audio
 	SIGNALPROCESSING_API void ArrayMinMaxNormalize(TArrayView<const float> InView, TArray<float>& OutArray);
 
 	// Scale an array to fit within [-1.f, 1.f]
-	SIGNALPROCESSING_API float ArrayMaxAbsValue(const TArrayView<float> InView);
+	SIGNALPROCESSING_API float ArrayMaxAbsValue(const TArrayView<const float> InView);
 
 	/** Multiply the second buffer in place by the first buffer. */
 	SIGNALPROCESSING_API void ArrayMultiplyInPlace(TArrayView<const float> InFloatBuffer, TArrayView<float> BufferToMultiply);
