Fix some audio parameters not being added when merging

diff --git a/Engine/Source/Runtime/AudioExtensions/Private/AudioParameter.cpp b/Engine/Source/Runtime/AudioExtensions/Private/AudioParameter.cpp
--- a/Engine/Source/Runtime/AudioExtensions/Private/AudioParameter.cpp
+++ b/Engine/Source/Runtime/AudioExtensions/Private/AudioParameter.cpp
@@ -211,6 +211,7 @@ void FAudioParameter::Merge(TArray<FAudioParameter>&& InParams, TArray<FAudioPar
 		while (!InParams.IsEmpty())
 		{
 			FAudioParameter& OutParam = OutParams[i];
+			// Break going through in params as there will be no more matches with this current out param 
 			if (InParams.Last().ParamName.FastLess(OutParam.ParamName))
 			{
 				break;
@@ -228,4 +229,7 @@ void FAudioParameter::Merge(TArray<FAudioParameter>&& InParams, TArray<FAudioPar
 			}
 		}
 	}
+
+	// Add the rest of the in params that have no matching out param 
+	OutParams.Append(MoveTemp(InParams));
 }
