Add first and last edited sample indexes to FWaveformTransformationWaveInfo #preflight 62fb5c49cf7afa018b8e6f3e

diff --git a/Engine/Source/Runtime/AudioExtensions/Public/IWaveformTransformation.h b/Engine/Source/Runtime/AudioExtensions/Public/IWaveformTransformation.h
--- a/Engine/Source/Runtime/AudioExtensions/Public/IWaveformTransformation.h
+++ b/Engine/Source/Runtime/AudioExtensions/Public/IWaveformTransformation.h
@@ -22,6 +22,8 @@ namespace Audio
 		float SampleRate = 0.f;
 		int32 NumChannels = 0;
 		Audio::FAlignedFloatBuffer* Audio = nullptr;
+		uint32 StartFrameOffset = 0;
+		uint32 NumEditedSamples = 0;
 	};
 
 	/*
