Moved AUDIO_SPATIALIZATION_PLUGIN_LLM_SCOPE from AudioDevice.h to IAudioExtensionPlugin.h to improve compile times because AudioDevice.h is expensive to include.

diff --git a/Engine/Plugins/Runtime/ResonanceAudio/Source/ResonanceAudio/Private/ResonanceAudioAmbisonicsSettings.cpp b/Engine/Plugins/Runtime/ResonanceAudio/Source/ResonanceAudio/Private/ResonanceAudioAmbisonicsSettings.cpp
--- a/Engine/Plugins/Runtime/ResonanceAudio/Source/ResonanceAudio/Private/ResonanceAudioAmbisonicsSettings.cpp
+++ b/Engine/Plugins/Runtime/ResonanceAudio/Source/ResonanceAudio/Private/ResonanceAudioAmbisonicsSettings.cpp
@@ -4,7 +4,7 @@
 
 #include "ResonanceAudioAmbisonicsSettings.h"
 
-#include "AudioDevice.h"
+#include "IAudioExtensionPlugin.h"
 #include "HAL/LowLevelMemTracker.h"
 #include "ResonanceAudioAmbisonicsSettingsProxy.h"
 #include "ResonanceAudioCommon.h"
diff --git a/Engine/Plugins/Runtime/ResonanceAudio/Source/ResonanceAudio/Private/ResonanceAudioAmbisonicsSettingsProxy.h b/Engine/Plugins/Runtime/ResonanceAudio/Source/ResonanceAudio/Private/ResonanceAudioAmbisonicsSettingsProxy.h
--- a/Engine/Plugins/Runtime/ResonanceAudio/Source/ResonanceAudio/Private/ResonanceAudioAmbisonicsSettingsProxy.h
+++ b/Engine/Plugins/Runtime/ResonanceAudio/Source/ResonanceAudio/Private/ResonanceAudioAmbisonicsSettingsProxy.h
@@ -6,7 +6,7 @@
 
 #include "ResonanceAudioAmbisonicsSettings.h"
 
-#include "AudioDevice.h"
+#include "IAudioExtensionPlugin.h"
 #include "HAL/LowLevelMemTracker.h"
 #include "ResonanceAudioCommon.h"
 
diff --git a/Engine/Source/Runtime/AudioExtensions/Private/AudioExtensionsModule.cpp b/Engine/Source/Runtime/AudioExtensions/Private/AudioExtensionsModule.cpp
--- a/Engine/Source/Runtime/AudioExtensions/Private/AudioExtensionsModule.cpp
+++ b/Engine/Source/Runtime/AudioExtensions/Private/AudioExtensionsModule.cpp
@@ -2,9 +2,14 @@
 
 #include "Modules/ModuleManager.h"
 
+#include "HAL/LowLevelMemStats.h"
 #include "IAudioCodecRegistry.h"
+#include "IAudioExtensionPlugin.h"
 #include "PcmCodec.h"
 
+DECLARE_LLM_MEMORY_STAT(TEXT("AudioSpatializationPlugins"), STAT_AudioSpatializationPluginsLLM, STATGROUP_LLMFULL);
+LLM_DEFINE_TAG(Audio_SpatializationPlugins, NAME_None, TEXT("Audio"), GET_STATFNAME(STAT_AudioSpatializationPluginsLLM), GET_STATFNAME(STAT_AudioSummaryLLM));
+
 class FAudioExtensionsModule : public IModuleInterface
 {
 public:
diff --git a/Engine/Source/Runtime/AudioExtensions/Public/IAudioExtensionPlugin.h b/Engine/Source/Runtime/AudioExtensions/Public/IAudioExtensionPlugin.h
--- a/Engine/Source/Runtime/AudioExtensions/Public/IAudioExtensionPlugin.h
+++ b/Engine/Source/Runtime/AudioExtensions/Public/IAudioExtensionPlugin.h
@@ -3,6 +3,7 @@
 
 #include "CoreMinimal.h"
 #include "Features/IModularFeature.h"
+#include "HAL/LowLevelMemTracker.h"
 #include "ISoundfieldFormat.h"
 #include "Math/Interval.h"
 #include "Modules/ModuleInterface.h"
@@ -13,7 +14,6 @@
 
 #include "IAudioExtensionPlugin.generated.h"
 
-
 // Forward Declarations
 class FAudioDevice;
 class FSoundEffectBase;
@@ -28,6 +28,9 @@ class IAudioSourceDataOverride;
 class IAudioSpatialization;
 class USoundSubmix;
 
+LLM_DECLARE_TAG_API(Audio_SpatializationPlugins, AUDIOEXTENSIONS_API);
+// Convenience macro for Audio_SpatializationPlugins LLM scope to avoid misspells.
+#define AUDIO_SPATIALIZATION_PLUGIN_LLM_SCOPE LLM_SCOPE_BYTAG(Audio_SpatializationPlugins);
 
 /**
 * Enumeration of audio plugin types
diff --git a/Engine/Source/Runtime/Engine/Engine.Build.cs b/Engine/Source/Runtime/Engine/Engine.Build.cs
--- a/Engine/Source/Runtime/Engine/Engine.Build.cs
+++ b/Engine/Source/Runtime/Engine/Engine.Build.cs
@@ -143,7 +143,8 @@ public class Engine : ModuleRules
 				"TraceLog",
 				"ColorManagement",
 				"Icmp",
-				"XmlParser"
+				"XmlParser",
+				"AudioExtensions"
 			}
 		);
 
diff --git a/Engine/Source/Runtime/Engine/Private/AudioDevice.cpp b/Engine/Source/Runtime/Engine/Private/AudioDevice.cpp
--- a/Engine/Source/Runtime/Engine/Private/AudioDevice.cpp
+++ b/Engine/Source/Runtime/Engine/Private/AudioDevice.cpp
@@ -53,9 +53,6 @@
 #include "Editor/EditorEngine.h"
 #endif // WITH_EDITOR
 
-DECLARE_LLM_MEMORY_STAT(TEXT("AudioSpatializationPlugins"), STAT_AudioSpatializationPluginsLLM, STATGROUP_LLMFULL);
-LLM_DEFINE_TAG(Audio_SpatializationPlugins, NAME_None, TEXT("Audio"), GET_STATFNAME(STAT_AudioSpatializationPluginsLLM), GET_STATFNAME(STAT_AudioSummaryLLM));
-
 static int32 AudioChannelCountCVar = 0;
 FAutoConsoleVariableRef CVarSetAudioChannelCount(
 	TEXT("au.SetAudioChannelCount"),
diff --git a/Engine/Source/Runtime/Engine/Public/AudioDevice.h b/Engine/Source/Runtime/Engine/Public/AudioDevice.h
--- a/Engine/Source/Runtime/Engine/Public/AudioDevice.h
+++ b/Engine/Source/Runtime/Engine/Public/AudioDevice.h
@@ -67,10 +67,6 @@ struct FAudioComponentParam;
 struct FAudioQualitySettings;
 struct FWaveInstance;
 
-LLM_DECLARE_TAG_API(Audio_SpatializationPlugins, ENGINE_API);
-// Convenience macro for Audio_SpatializationPlugins LLM scope to avoid misspells.
-#define AUDIO_SPATIALIZATION_PLUGIN_LLM_SCOPE LLM_SCOPE_BYTAG(Audio_SpatializationPlugins);
-
 namespace Audio
 {
 	class FAudioDebugger;
