Fix crash from malformed printf
#rnx

diff --git a/Engine/Source/Runtime/AudioExtensions/Public/IAudioCodec.h b/Engine/Source/Runtime/AudioExtensions/Public/IAudioCodec.h
--- a/Engine/Source/Runtime/AudioExtensions/Public/IAudioCodec.h
+++ b/Engine/Source/Runtime/AudioExtensions/Public/IAudioCodec.h
@@ -748,7 +748,7 @@ namespace Audio
 		uint32 Version = 0;						// For Compatibility (all codecs in this family will be compatible with this data).	
 		FCodecFeatures Features;				// If defined, is the feature set for this codec 	
 
-		FString ToString() const { return FString::Printf(TEXT("%s:%s:%d")); }
+		FString ToString() const { return FString::Printf(TEXT("%s:%s:%d"), *FamilyName.ToString(), *Name.ToString(), Version); }
 		
 		bool operator==(const FCodecDetails& Rhs) const 
 		{ 
