Removed TIsSame usage from AIModule (part of a larger change to deprecate TIsSame).

diff --git a/Engine/Source/Runtime/AIModule/Classes/BehaviorTree/BTAuxiliaryNode.h b/Engine/Source/Runtime/AIModule/Classes/BehaviorTree/BTAuxiliaryNode.h
--- a/Engine/Source/Runtime/AIModule/Classes/BehaviorTree/BTAuxiliaryNode.h
+++ b/Engine/Source/Runtime/AIModule/Classes/BehaviorTree/BTAuxiliaryNode.h
@@ -107,9 +107,9 @@ protected:
 	template<typename TickNode,	typename OnBecomeRelevant, typename OnCeaseRelevant>
 	void InitNotifyFlags(TickNode, OnBecomeRelevant, OnCeaseRelevant)
 	{
-		bNotifyTick = !TIsSame<decltype(&UBTAuxiliaryNode::TickNode), TickNode>::Value;
-		bNotifyBecomeRelevant = !TIsSame<decltype(&UBTAuxiliaryNode::OnBecomeRelevant), OnBecomeRelevant>::Value;
-		bNotifyCeaseRelevant = !TIsSame<decltype(&UBTAuxiliaryNode::OnCeaseRelevant), OnCeaseRelevant>::Value;
+		bNotifyTick = !std::is_same_v<decltype(&UBTAuxiliaryNode::TickNode), TickNode>;
+		bNotifyBecomeRelevant = !std::is_same_v<decltype(&UBTAuxiliaryNode::OnBecomeRelevant), OnBecomeRelevant>;
+		bNotifyCeaseRelevant = !std::is_same_v<decltype(&UBTAuxiliaryNode::OnCeaseRelevant), OnCeaseRelevant>;
 	}
 };
 
diff --git a/Engine/Source/Runtime/AIModule/Classes/BehaviorTree/BTCompositeNode.h b/Engine/Source/Runtime/AIModule/Classes/BehaviorTree/BTCompositeNode.h
--- a/Engine/Source/Runtime/AIModule/Classes/BehaviorTree/BTCompositeNode.h
+++ b/Engine/Source/Runtime/AIModule/Classes/BehaviorTree/BTCompositeNode.h
@@ -273,12 +273,12 @@ protected:
 	void InitNotifyFlags(NotifyChildExecution, NotifyNodeActivation, NotifyNodeDeactivation,
 						 CanNotifyDecoratorsOnActivation, CanNotifyDecoratorsOnDeactivation, CanNotifyDecoratorsOnFailedActivation)
 	{
-		bUseChildExecutionNotify = !TIsSame<decltype(&UBTCompositeNode::NotifyChildExecution), NotifyChildExecution>::Value;
-		bUseNodeActivationNotify = !TIsSame<decltype(&UBTCompositeNode::NotifyNodeActivation), NotifyNodeActivation>::Value;
-		bUseNodeDeactivationNotify = !TIsSame<decltype(&UBTCompositeNode::NotifyNodeDeactivation), NotifyNodeDeactivation>::Value;
-		bUseDecoratorsActivationCheck = !TIsSame<decltype(&UBTCompositeNode::CanNotifyDecoratorsOnActivation), CanNotifyDecoratorsOnActivation>::Value;
-		bUseDecoratorsDeactivationCheck = !TIsSame<decltype(&UBTCompositeNode::CanNotifyDecoratorsOnDeactivation), CanNotifyDecoratorsOnDeactivation>::Value;
-		bUseDecoratorsFailedActivationCheck = !TIsSame<decltype(&UBTCompositeNode::CanNotifyDecoratorsOnFailedActivation), CanNotifyDecoratorsOnFailedActivation>::Value;
+		bUseChildExecutionNotify = !std::is_same_v<decltype(&UBTCompositeNode::NotifyChildExecution), NotifyChildExecution>;
+		bUseNodeActivationNotify = !std::is_same_v<decltype(&UBTCompositeNode::NotifyNodeActivation), NotifyNodeActivation>;
+		bUseNodeDeactivationNotify = !std::is_same_v<decltype(&UBTCompositeNode::NotifyNodeDeactivation), NotifyNodeDeactivation>;
+		bUseDecoratorsActivationCheck = !std::is_same_v<decltype(&UBTCompositeNode::CanNotifyDecoratorsOnActivation), CanNotifyDecoratorsOnActivation>;
+		bUseDecoratorsDeactivationCheck = !std::is_same_v<decltype(&UBTCompositeNode::CanNotifyDecoratorsOnDeactivation), CanNotifyDecoratorsOnDeactivation>;
+		bUseDecoratorsFailedActivationCheck = !std::is_same_v<decltype(&UBTCompositeNode::CanNotifyDecoratorsOnFailedActivation), CanNotifyDecoratorsOnFailedActivation>;
 	}
 };
 
diff --git a/Engine/Source/Runtime/AIModule/Classes/BehaviorTree/BTDecorator.h b/Engine/Source/Runtime/AIModule/Classes/BehaviorTree/BTDecorator.h
--- a/Engine/Source/Runtime/AIModule/Classes/BehaviorTree/BTDecorator.h
+++ b/Engine/Source/Runtime/AIModule/Classes/BehaviorTree/BTDecorator.h
@@ -131,12 +131,12 @@ protected:
 	void InitNotifyFlags(TickNode, OnBecomeRelevant, OnCeaseRelevant,
 					 OnNodeActivation, OnNodeDeactivation, OnNodeProcessed)
 	{
-		bNotifyTick = !TIsSame<decltype(&UBTDecorator::TickNode), TickNode>::Value;
-		bNotifyBecomeRelevant = !TIsSame<decltype(&UBTDecorator::OnBecomeRelevant), OnBecomeRelevant>::Value;
-		bNotifyCeaseRelevant = !TIsSame<decltype(&UBTDecorator::OnCeaseRelevant), OnCeaseRelevant>::Value;
-		bNotifyActivation = !TIsSame<decltype(&UBTDecorator::OnNodeActivation), OnNodeActivation>::Value;
-		bNotifyDeactivation = !TIsSame<decltype(&UBTDecorator::OnNodeDeactivation), OnNodeDeactivation>::Value;
-		bNotifyProcessed = !TIsSame<decltype(&UBTDecorator::OnNodeProcessed), OnNodeProcessed>::Value;
+		bNotifyTick = !std::is_same_v<decltype(&UBTDecorator::TickNode), TickNode>;
+		bNotifyBecomeRelevant = !std::is_same_v<decltype(&UBTDecorator::OnBecomeRelevant), OnBecomeRelevant>;
+		bNotifyCeaseRelevant = !std::is_same_v<decltype(&UBTDecorator::OnCeaseRelevant), OnCeaseRelevant>;
+		bNotifyActivation = !std::is_same_v<decltype(&UBTDecorator::OnNodeActivation), OnNodeActivation>;
+		bNotifyDeactivation = !std::is_same_v<decltype(&UBTDecorator::OnNodeDeactivation), OnNodeDeactivation>;
+		bNotifyProcessed = !std::is_same_v<decltype(&UBTDecorator::OnNodeProcessed), OnNodeProcessed>;
 	}
 };
 
diff --git a/Engine/Source/Runtime/AIModule/Classes/BehaviorTree/BTService.h b/Engine/Source/Runtime/AIModule/Classes/BehaviorTree/BTService.h
--- a/Engine/Source/Runtime/AIModule/Classes/BehaviorTree/BTService.h
+++ b/Engine/Source/Runtime/AIModule/Classes/BehaviorTree/BTService.h
@@ -88,10 +88,10 @@ protected:
 	template<typename TickNode, typename OnBecomeRelevant, typename OnCeaseRelevant, typename OnSearchStart>
 	void InitNotifyFlags(TickNode, OnBecomeRelevant, OnCeaseRelevant, OnSearchStart)
 	{
-		bNotifyTick = !TIsSame<decltype(&UBTService::TickNode), TickNode>::Value;
-		bNotifyBecomeRelevant = !TIsSame<decltype(&UBTService::OnBecomeRelevant), OnBecomeRelevant>::Value;
-		bNotifyCeaseRelevant = !TIsSame<decltype(&UBTService::OnCeaseRelevant), OnCeaseRelevant>::Value;
-		bNotifyOnSearch = !TIsSame<decltype(&UBTService::OnSearchStart), OnSearchStart>::Value;
+		bNotifyTick = !std::is_same_v<decltype(&UBTService::TickNode), TickNode>;
+		bNotifyBecomeRelevant = !std::is_same_v<decltype(&UBTService::OnBecomeRelevant), OnBecomeRelevant>;
+		bNotifyCeaseRelevant = !std::is_same_v<decltype(&UBTService::OnCeaseRelevant), OnCeaseRelevant>;
+		bNotifyOnSearch = !std::is_same_v<decltype(&UBTService::OnSearchStart), OnSearchStart>;
 	}
 };
 
diff --git a/Engine/Source/Runtime/AIModule/Classes/BehaviorTree/BTTaskNode.h b/Engine/Source/Runtime/AIModule/Classes/BehaviorTree/BTTaskNode.h
--- a/Engine/Source/Runtime/AIModule/Classes/BehaviorTree/BTTaskNode.h
+++ b/Engine/Source/Runtime/AIModule/Classes/BehaviorTree/BTTaskNode.h
@@ -127,8 +127,8 @@ protected:
 	template<typename TickTask,	typename OnTaskFinished>
 	void InitNotifyFlags(TickTask, OnTaskFinished)
 	{
-		bNotifyTick = !TIsSame<decltype(&UBTTaskNode::TickTask), TickTask>::Value;
-		bNotifyTaskFinished = !TIsSame<decltype(&UBTTaskNode::OnTaskFinished), OnTaskFinished>::Value;
+		bNotifyTick = !std::is_same_v<decltype(&UBTTaskNode::TickTask), TickTask>;
+		bNotifyTaskFinished = !std::is_same_v<decltype(&UBTTaskNode::OnTaskFinished), OnTaskFinished>;
 	}
 
 	/** if set, conditional tick will use remaining time from node's memory */
diff --git a/Engine/Source/Runtime/Analytics/Analytics/Public/AnalyticsEventAttribute.h b/Engine/Source/Runtime/Analytics/Analytics/Public/AnalyticsEventAttribute.h
--- a/Engine/Source/Runtime/Analytics/Analytics/Public/AnalyticsEventAttribute.h
+++ b/Engine/Source/Runtime/Analytics/Analytics/Public/AnalyticsEventAttribute.h
@@ -145,7 +145,7 @@ inline FAnalyticsEventAttribute::FAnalyticsEventAttribute(FString InName, ValueT
 , AttrValueString(AnalyticsConversionToString(Forward<ValueType>(InValue)))
 , AttrValueNumber(0)
 , AttrValueBool(false)
-, AttrType(TIsArithmetic<typename TDecay<ValueType>::Type>::Value || TIsSame<typename TDecay<ValueType>::Type, FJsonNull>::Value || TIsSame<typename TDecay<ValueType>::Type, FJsonFragment>::Value ? AttrTypeEnum::JsonFragment : AttrTypeEnum::String)
+, AttrType(TIsArithmetic<typename TDecay<ValueType>::Type>::Value || std::is_same_v<typename TDecay<ValueType>::Type, FJsonNull> || std::is_same_v<typename TDecay<ValueType>::Type, FJsonFragment> ? AttrTypeEnum::JsonFragment : AttrTypeEnum::String)
 {
 
 }
@@ -206,7 +206,7 @@ template<typename ValueType>
 inline void FAnalyticsEventAttribute::SetValue(ValueType&& InValue)
 {
 	const_cast<FString&>(AttrValueString) = AnalyticsConversionToString(Forward<ValueType>(InValue));
-	const_cast<AttrTypeEnum&>(AttrType) = TIsArithmetic<typename TDecay<ValueType>::Type>::Value || TIsSame<typename TDecay<ValueType>::Type, FJsonNull>::Value || TIsSame<typename TDecay<ValueType>::Type, FJsonFragment>::Value ? AttrTypeEnum::JsonFragment : AttrTypeEnum::String;
+	const_cast<AttrTypeEnum&>(AttrType) = TIsArithmetic<typename TDecay<ValueType>::Type>::Value || std::is_same_v<typename TDecay<ValueType>::Type, FJsonNull> || std::is_same_v<typename TDecay<ValueType>::Type, FJsonFragment> ? AttrTypeEnum::JsonFragment : AttrTypeEnum::String;
 }
 
 inline void FAnalyticsEventAttribute::SwitchToJsonFragment()
diff --git a/Engine/Source/Runtime/AudioExtensions/Public/IAudioCodec.h b/Engine/Source/Runtime/AudioExtensions/Public/IAudioCodec.h
--- a/Engine/Source/Runtime/AudioExtensions/Public/IAudioCodec.h
+++ b/Engine/Source/Runtime/AudioExtensions/Public/IAudioCodec.h
@@ -96,18 +96,22 @@ namespace Audio
 	template<typename From, typename To>
 	static constexpr float GetConvertScalar()
 	{
-		if (TIsSame<From, To>::Value)
+		if constexpr (std::is_same_v<From, To>)
 		{
 			return 1.0f;
 		}
-		else if (TIsSame<From, int16>::Value && TIsSame<To, float>::Value)
+		else if constexpr (std::is_same_v<From, int16> && std::is_same_v<To, float>)
 		{
 			return 1.f / (static_cast<float>(TNumericLimits<int16>::Max()) + 1.f);
 		}
-		else if (TIsSame<From, float>::Value && TIsSame<To, int16>::Value)
+		else if constexpr (std::is_same_v<From, float> && std::is_same_v<To, int16>)
 		{
 			return static_cast<float>(TNumericLimits<int16>::Max()) + 1.f;
 		}
+		else
+		{
+			static_assert(sizeof(From) == 0, "Unsupported conversion");
+		}
 	}
 		
 	enum EBitRepresentation
@@ -490,7 +494,7 @@ namespace Audio
 			int32 AvailSpace = Buffer.Num() - Offset;
 			int32 Num = FMath::Min(AvailSpace, InAudio.Num());
 
-			if (TIsSame<TSampleType, T>::Value)
+			if constexpr (std::is_same_v<TSampleType, T>)
 			{
 				FMemory::Memcpy(Dst, Src, Num * sizeof(T));
 			}
@@ -526,7 +530,7 @@ namespace Audio
 			TArrayView<T> InExternalBuffer )
 		{
 			check(InExternalBuffer.Num() > 0);
-			if (TIsSame<T, TSampleType>::Value && Offset > 0)
+			if (std::is_same_v<T, TSampleType> && Offset > 0)
 			{
 				int32 NumSamplesPopped = Offset;
 				Offset = 0;
@@ -589,7 +593,7 @@ namespace Audio
 		int32 PushAudioInternal(
 			const FPushedAudioDetails& InDetails, TArrayView<const T> InBuffer) 
 		{
-			if (TIsSame<T, TSampleType>::Value)
+			if constexpr (std::is_same_v<T, TSampleType>)
 			{
 				return Buffer.Push(reinterpret_cast<const TSampleType*>(InBuffer.GetData()), InBuffer.Num());
 			}
@@ -620,7 +624,7 @@ namespace Audio
 		}
 		int32 PopAudio(TArrayView<int16> InExternalInt16Buffer, FPushedAudioDetails& OutDetails) override
 		{
-			if( TIsSame<TSampleType, int16>::Value )
+			if constexpr (std::is_same_v<TSampleType, int16>)
 			{
 				return Buffer.Pop(reinterpret_cast<TSampleType*>(InExternalInt16Buffer.GetData()), InExternalInt16Buffer.Num());
 			}
@@ -630,7 +634,7 @@ namespace Audio
 
 		int32 PopAudio(TArrayView<float> InExternalFloat32Buffer, FPushedAudioDetails& OutDetails) override
 		{
-			if( TIsSame<TSampleType, float>::Value )
+			if constexpr (std::is_same_v<TSampleType, float>)
 			{
 				return Buffer.Pop(reinterpret_cast<TSampleType*>(InExternalFloat32Buffer.GetData()), InExternalFloat32Buffer.Num());
 			}
diff --git a/Engine/Source/Runtime/ColorManagement/Public/ColorSpace.h b/Engine/Source/Runtime/ColorManagement/Public/ColorSpace.h
--- a/Engine/Source/Runtime/ColorManagement/Public/ColorSpace.h
+++ b/Engine/Source/Runtime/ColorManagement/Public/ColorSpace.h
@@ -68,7 +68,7 @@ FORCEINLINE FVector3d ToXYZ(const FVector2d& Coordinate)
 template<typename T>
 FORCEINLINE UE::Math::TMatrix<T> Transpose(const FMatrix44d& Transform)
 {
-	if constexpr (TIsSame<T, double>::Value)
+	if constexpr (std::is_same_v<T, double>)
 	{
 		return Transform.GetTransposed();
 	}
