Fix default values not being retained when sound is virtualized.
#rnx

diff --git a/Engine/Plugins/Runtime/Metasound/Source/MetasoundFrontend/Private/MetasoundParameterTransmitter.cpp b/Engine/Plugins/Runtime/Metasound/Source/MetasoundFrontend/Private/MetasoundParameterTransmitter.cpp
--- a/Engine/Plugins/Runtime/Metasound/Source/MetasoundFrontend/Private/MetasoundParameterTransmitter.cpp
+++ b/Engine/Plugins/Runtime/Metasound/Source/MetasoundFrontend/Private/MetasoundParameterTransmitter.cpp
@@ -248,10 +248,8 @@ namespace Metasound
 	{
 	}
 
-	bool FMetaSoundParameterTransmitter::Reset()
+	void FMetaSoundParameterTransmitter::OnDeleteActiveSound()
 	{
-		bool bSuccess = true;
-
 		for (const FSendInfo& SendInfo : SendInfos)
 		{
 			if (InputSends.Remove(SendInfo.ParameterName))
@@ -261,13 +259,17 @@ namespace Metasound
 				// multiple times. Multiple removals of data channels has caused
 				// race conditions between newly created transmitters and transmitters
 				// being cleaned up.
-				bSuccess &= FDataTransmissionCenter::Get().UnregisterDataChannel(SendInfo.Address);
+				FDataTransmissionCenter::Get().UnregisterDataChannel(SendInfo.Address);
 			}
 		}
 
-		bSuccess &= Audio::FParameterTransmitterBase::Reset();
-
-		return bSuccess;
+		Audio::FParameterTransmitterBase::OnDeleteActiveSound();
+	}
+	
+	bool FMetaSoundParameterTransmitter::Reset()
+	{
+		OnDeleteActiveSound();
+		return true;
 	}
 
 	bool FMetaSoundParameterTransmitter::SetParameters(TArray<FAudioParameter>&& InParameters)
diff --git a/Engine/Plugins/Runtime/Metasound/Source/MetasoundFrontend/Public/MetasoundParameterTransmitter.h b/Engine/Plugins/Runtime/Metasound/Source/MetasoundFrontend/Public/MetasoundParameterTransmitter.h
--- a/Engine/Plugins/Runtime/Metasound/Source/MetasoundFrontend/Public/MetasoundParameterTransmitter.h
+++ b/Engine/Plugins/Runtime/Metasound/Source/MetasoundFrontend/Public/MetasoundParameterTransmitter.h
@@ -88,8 +88,11 @@ namespace Metasound
 		FMetaSoundParameterTransmitter(FMetaSoundParameterTransmitter::FInitParams&& InInitParams);
 		virtual ~FMetaSoundParameterTransmitter() = default;
 
-		bool Reset() override;
+		virtual void OnDeleteActiveSound() override;
 
+		UE_DEPRECATED(5.2, "Use ResetParameters() or OnDeleteActiveSound() instead depending on use case.")
+		bool Reset() override;
+		
 		/** Sets parameters using array of AudioParameter structs
 		 *
 		 * @param InParameter - Parameter to set.
diff --git a/Engine/Source/Runtime/AudioExtensions/Private/IAudioParameterTransmitter.cpp b/Engine/Source/Runtime/AudioExtensions/Private/IAudioParameterTransmitter.cpp
--- a/Engine/Source/Runtime/AudioExtensions/Private/IAudioParameterTransmitter.cpp
+++ b/Engine/Source/Runtime/AudioExtensions/Private/IAudioParameterTransmitter.cpp
@@ -28,20 +28,25 @@ namespace Audio
 		return false;
 	}
 
-	const TArray<FAudioParameter>& FParameterTransmitterBase::GetParameters() const
+	void FParameterTransmitterBase::ResetParameters() 
 	{
-		return AudioParameters;
+		AudioParameters.Reset();
 	}
-
+	
 	bool FParameterTransmitterBase::Reset()
 	{
-		AudioParameters.Reset();
+		ResetParameters();
 		return true;
 	}
 
+	const TArray<FAudioParameter>& FParameterTransmitterBase::GetParameters() const
+	{
+		return AudioParameters;
+	}
+
 	bool FParameterTransmitterBase::SetParameters(TArray<FAudioParameter>&& InParameters)
 	{
 		FAudioParameter::Merge(MoveTemp(InParameters), AudioParameters);
 		return true;
 	}
-} // namespace Audio
\ No newline at end of file
+} // namespace Audio
diff --git a/Engine/Source/Runtime/AudioExtensions/Public/IAudioParameterTransmitter.h b/Engine/Source/Runtime/AudioExtensions/Public/IAudioParameterTransmitter.h
--- a/Engine/Source/Runtime/AudioExtensions/Public/IAudioParameterTransmitter.h
+++ b/Engine/Source/Runtime/AudioExtensions/Public/IAudioParameterTransmitter.h
@@ -51,7 +51,15 @@ namespace Audio
 
 			virtual ~IParameterTransmitter() = default;
 
-			virtual bool Reset() = 0;
+			UE_DEPRECATED(5.2, "Use ResetParameters() or OnDeleteActiveSound() instead depending on use case.")
+			virtual bool Reset() { ResetParameters(); return true; }
+
+			// Reset parameters which stored on the transmitter.
+			virtual void ResetParameters() {}
+
+			// Called when the active sound is deleted due to the sound finishing,
+			// being stopped, or being virtualized. 
+			virtual void OnDeleteActiveSound() {}
 
 			// Return the cached parameter with the given name if it exists
 			// @return False if param not found, true if found.
@@ -74,7 +82,9 @@ namespace Audio
 		virtual ~FParameterTransmitterBase() = default;
 
 		virtual bool GetParameter(FName InName, FAudioParameter& OutParam) const override;
+		virtual void ResetParameters() override;
 		virtual const TArray<FAudioParameter>& GetParameters() const override;
+		UE_DEPRECATED(5.2, "Use ResetParameters() or OnDeleteActiveSound() instead depending on use case.")
 		virtual bool Reset() override;
 		virtual bool SetParameters(TArray<FAudioParameter>&& InParameters) override;
 
diff --git a/Engine/Source/Runtime/Engine/Private/Audio/SoundParameterControllerInterface.cpp b/Engine/Source/Runtime/Engine/Private/Audio/SoundParameterControllerInterface.cpp
--- a/Engine/Source/Runtime/Engine/Private/Audio/SoundParameterControllerInterface.cpp
+++ b/Engine/Source/Runtime/Engine/Private/Audio/SoundParameterControllerInterface.cpp
@@ -28,7 +28,7 @@ void ISoundParameterControllerInterface::ResetParameters()
 			{
 				if (Audio::IParameterTransmitter* Transmitter = ActiveSound.GetTransmitter())
 				{
-					Transmitter->Reset();
+					Transmitter->ResetParameters();
 				}
 			}, GET_STATID(STAT_AudioResetParameters));
 		}
diff --git a/Engine/Source/Runtime/Engine/Private/AudioDevice.cpp b/Engine/Source/Runtime/Engine/Private/AudioDevice.cpp
--- a/Engine/Source/Runtime/Engine/Private/AudioDevice.cpp
+++ b/Engine/Source/Runtime/Engine/Private/AudioDevice.cpp
@@ -5478,11 +5478,11 @@ void FAudioDevice::ProcessingPendingActiveSoundStops(bool bForceDelete)
 				ActiveSound->bAsyncOcclusionPending = false;
 				PendingSoundsToDelete.RemoveAtSwap(i, 1, false);
 
-				// TODO: Phil Popp investigate a more robust fix -- call Shutdown in FActiveSound or transmitter destructor?
 				if (Audio::IParameterTransmitter* Transmitter = ActiveSound->GetTransmitter())
 				{
-					Transmitter->Reset();
+					Transmitter->OnDeleteActiveSound();
 				}
+				ActiveSound->ClearTransmitter();
 
 				NotifyPendingDeleteInternal(*ActiveSound);
 				delete ActiveSound;
@@ -5538,11 +5538,11 @@ void FAudioDevice::ProcessingPendingActiveSoundStops(bool bForceDelete)
 
 			if (bDeleteActiveSound)
 			{
-				// TODO: Phil Popp investigate a more robust fix -- call Shutdown in FActiveSound or transmitter destructor?
 				if (Audio::IParameterTransmitter* Transmitter = ActiveSound->GetTransmitter())
 				{
-					Transmitter->Reset();
+					Transmitter->OnDeleteActiveSound();
 				}
+				ActiveSound->ClearTransmitter();
 
 				NotifyPendingDeleteInternal(*ActiveSound);
 
