Deprecate IAudioCodec API

diff --git a/Engine/Plugins/Runtime/Metasound/Source/MetasoundEngine/Private/MetasoundWaveProxyReader.h b/Engine/Plugins/Runtime/Metasound/Source/MetasoundEngine/Private/MetasoundWaveProxyReader.h
--- a/Engine/Plugins/Runtime/Metasound/Source/MetasoundEngine/Private/MetasoundWaveProxyReader.h
+++ b/Engine/Plugins/Runtime/Metasound/Source/MetasoundEngine/Private/MetasoundWaveProxyReader.h
@@ -15,10 +15,12 @@ namespace Metasound
 {
 	namespace MetasoundWaveProxyReaderPrivate
 	{
+		PRAGMA_DISABLE_DEPRECATION_WARNINGS
+
 		/** FDecoderOutput maintains a circular buffer of audio which is produced
 		 * from the decoder. 
 		 */
-		class FDecoderOutput : public Audio::IDecoderOutput
+		class UE_DEPRECATED(5.3, "Please use IAudioFormat/IStreamedCompressedInfo instead") FDecoderOutput : public Audio::IDecoderOutput
 		{
 			static constexpr uint32 MinNumFramesPerDecode = 1;
 			static constexpr uint32 DefaultNumChannels = 1;
@@ -102,7 +104,7 @@ namespace Metasound
 			Audio::FAlignedFloatBuffer SampleConversionBuffer;
 		};
 	}
-	
+		
 	/** FWaveProxyReader reads a FWaveProxy and outputs 32 bit interleaved audio.
 	 *
 	 * FWaveProxyReader provides controls for looping and relevant frame index
@@ -260,4 +262,6 @@ namespace Metasound
 		bool bIsDecoderValid = false;
 		bool bFallbackSeekMethodWarningLogged = false;
 	};
+
+	PRAGMA_ENABLE_DEPRECATION_WARNINGS
 }
diff --git a/Engine/Source/Runtime/AudioCodecEngine/Private/AudioCodecEngineModule.cpp b/Engine/Source/Runtime/AudioCodecEngine/Private/AudioCodecEngineModule.cpp
--- a/Engine/Source/Runtime/AudioCodecEngine/Private/AudioCodecEngineModule.cpp
+++ b/Engine/Source/Runtime/AudioCodecEngine/Private/AudioCodecEngineModule.cpp
@@ -5,6 +5,8 @@
 #include "IAudioCodecRegistry.h"
 #include "DecoderBackCompat.h"
 
+PRAGMA_DISABLE_DEPRECATION_WARNINGS
+
 class FAudioCodecEngineModule : public IModuleInterface
 {
 public:
@@ -26,4 +28,6 @@ public:
 	}
 };
 
+PRAGMA_ENABLE_DEPRECATION_WARNINGS
+
 IMPLEMENT_MODULE(FAudioCodecEngineModule, AudioCodecEngine);
diff --git a/Engine/Source/Runtime/AudioCodecEngine/Private/DecoderBackCompat.cpp b/Engine/Source/Runtime/AudioCodecEngine/Private/DecoderBackCompat.cpp
--- a/Engine/Source/Runtime/AudioCodecEngine/Private/DecoderBackCompat.cpp
+++ b/Engine/Source/Runtime/AudioCodecEngine/Private/DecoderBackCompat.cpp
@@ -6,6 +6,8 @@
 
 namespace Audio
 {
+	PRAGMA_DISABLE_DEPRECATION_WARNINGS
+
 	static const FCodecDetails Details =
 	{
 		TEXT("BackCompat"),
@@ -122,4 +124,5 @@ namespace Audio
 		}
 	}
 
+	PRAGMA_ENABLE_DEPRECATION_WARNINGS
 }
diff --git a/Engine/Source/Runtime/AudioCodecEngine/Private/DecoderBackCompat.h b/Engine/Source/Runtime/AudioCodecEngine/Private/DecoderBackCompat.h
--- a/Engine/Source/Runtime/AudioCodecEngine/Private/DecoderBackCompat.h
+++ b/Engine/Source/Runtime/AudioCodecEngine/Private/DecoderBackCompat.h
@@ -6,7 +6,9 @@
 
 namespace Audio
 {
-	struct FBackCompatCodec : public ICodec
+	PRAGMA_DISABLE_DEPRECATION_WARNINGS
+
+	struct UE_DEPRECATED(5.3, "Please use IAudioFormat/IStreamedCompressedInfo instead") FBackCompatCodec : public ICodec
 	{
 		// Static version.
 		static const FCodecDetails& GetDetailsStatic();
@@ -18,11 +20,10 @@ namespace Audio
 		// Factory for decoders
 		FDecoderPtr CreateDecoder(
 			IDecoder::FDecoderInputPtr InSrc,
-			IDecoder::FDecoderOutputPtr InDst) override;
-	
+			IDecoder::FDecoderOutputPtr InDst) override;	
 	};
 
-	struct FBackCompat : public IDecoder
+	struct UE_DEPRECATED(5.3, "Please use IAudioFormat/IStreamedCompressedInfo instead") FBackCompat : public IDecoder
 	{
 		IDecoder::FDecoderInputPtr Src = nullptr;
 		IDecoder::FDecoderOutputPtr Dst = nullptr;
@@ -46,4 +47,6 @@ namespace Audio
 		bool bPreviousIsStreaming = false;
 		bool bIsFirstDecode = true;
 	};
+
+	PRAGMA_ENABLE_DEPRECATION_WARNINGS
 }
diff --git a/Engine/Source/Runtime/AudioCodecEngine/Private/DecoderInputBackCompat.cpp b/Engine/Source/Runtime/AudioCodecEngine/Private/DecoderInputBackCompat.cpp
--- a/Engine/Source/Runtime/AudioCodecEngine/Private/DecoderInputBackCompat.cpp
+++ b/Engine/Source/Runtime/AudioCodecEngine/Private/DecoderInputBackCompat.cpp
@@ -10,6 +10,8 @@
 
 namespace Audio
 {	
+	PRAGMA_DISABLE_DEPRECATION_WARNINGS
+
 	ICompressedAudioInfo* FBackCompatInput::GetInfo(
 		FFormatDescriptorSection* OutDescriptor /*= nullptr*/) const
 	{
@@ -120,4 +122,7 @@ namespace Audio
 		audio_ensure(false);
 		return MakeArrayView<const uint8>(nullptr,0);
 	}
+
+	PRAGMA_ENABLE_DEPRECATION_WARNINGS
 }
+
diff --git a/Engine/Source/Runtime/AudioCodecEngine/Private/DecoderInputBackCompat.h b/Engine/Source/Runtime/AudioCodecEngine/Private/DecoderInputBackCompat.h
--- a/Engine/Source/Runtime/AudioCodecEngine/Private/DecoderInputBackCompat.h
+++ b/Engine/Source/Runtime/AudioCodecEngine/Private/DecoderInputBackCompat.h
@@ -10,6 +10,8 @@ class ICompressedAudioInfo; 	// Forward declares.
 
 namespace Audio
 {
+	PRAGMA_DISABLE_DEPRECATION_WARNINGS
+
 	struct FBackCompatInput : public IDecoderInput
 	{
 		FName OldFormatName;
@@ -42,4 +44,6 @@ namespace Audio
 		TArrayView<const uint8> PopNextPacket(
 			int32 InPacketSize) override;
 	};
+
+	PRAGMA_ENABLE_DEPRECATION_WARNINGS
 }
diff --git a/Engine/Source/Runtime/AudioCodecEngine/Private/DecoderInputFactory.cpp b/Engine/Source/Runtime/AudioCodecEngine/Private/DecoderInputFactory.cpp
--- a/Engine/Source/Runtime/AudioCodecEngine/Private/DecoderInputFactory.cpp
+++ b/Engine/Source/Runtime/AudioCodecEngine/Private/DecoderInputFactory.cpp
@@ -7,6 +7,8 @@
 
 namespace Audio
 {		
+	PRAGMA_DISABLE_DEPRECATION_WARNINGS
+
 	TUniquePtr<Audio::IDecoderInput> CreateBackCompatDecoderInput(
 		FName InOldFormatName,
 		const FSoundWaveProxyPtr& InSoundWave)
@@ -18,4 +20,7 @@ namespace Audio
 
 		return nullptr;
 	}
+	
+	
+	PRAGMA_ENABLE_DEPRECATION_WARNINGS
 }
diff --git a/Engine/Source/Runtime/AudioCodecEngine/Public/DecoderInputFactory.h b/Engine/Source/Runtime/AudioCodecEngine/Public/DecoderInputFactory.h
--- a/Engine/Source/Runtime/AudioCodecEngine/Public/DecoderInputFactory.h
+++ b/Engine/Source/Runtime/AudioCodecEngine/Public/DecoderInputFactory.h
@@ -11,7 +11,6 @@
 class USoundWave;
 class FSoundWaveProxy;
 
-
 namespace Audio
 {
 	// Forward declares.
@@ -19,5 +18,7 @@ namespace Audio
 	using FSoundWaveProxyPtr = TSharedPtr<FSoundWaveProxy, ESPMode::ThreadSafe>;
 
 	// Just loose for now.
+	PRAGMA_DISABLE_DEPRECATION_WARNINGS
 	AUDIOCODECENGINE_API TUniquePtr<IDecoderInput> CreateBackCompatDecoderInput(FName InOldFormatName, const FSoundWaveProxyPtr& InSoundWave);
+	PRAGMA_DISABLE_DEPRECATION_WARNINGS
 }
diff --git a/Engine/Source/Runtime/AudioExtensions/Private/AudioExtensionsModule.cpp b/Engine/Source/Runtime/AudioExtensions/Private/AudioExtensionsModule.cpp
--- a/Engine/Source/Runtime/AudioExtensions/Private/AudioExtensionsModule.cpp
+++ b/Engine/Source/Runtime/AudioExtensions/Private/AudioExtensionsModule.cpp
@@ -4,9 +4,7 @@
 
 #include "AudioExtensionsLog.h"
 #include "HAL/LowLevelMemStats.h"
-#include "IAudioCodecRegistry.h"
 #include "IAudioExtensionPlugin.h"
-#include "PcmCodec.h"
 
 DECLARE_LLM_MEMORY_STAT(TEXT("AudioSpatializationPlugins"), STAT_AudioSpatializationPluginsLLM, STATGROUP_LLMFULL);
 LLM_DEFINE_TAG(Audio_SpatializationPlugins, NAME_None, TEXT("Audio"), GET_STATFNAME(STAT_AudioSpatializationPluginsLLM), GET_STATFNAME(STAT_AudioSummaryLLM));
@@ -20,9 +18,6 @@ public:
 	{
 		FModuleManager::Get().LoadModuleChecked(TEXT("SignalProcessing"));
 		FModuleManager::Get().LoadModuleChecked(TEXT("AudioMixerCore"));
-
-		// Built in Codecs.
-		Audio::ICodecRegistry::Get().RegisterCodec(Audio::Create_PcmCodec());
 	}
 };
 
diff --git a/Engine/Source/Runtime/AudioExtensions/Private/IAudioCodec.cpp b/Engine/Source/Runtime/AudioExtensions/Private/IAudioCodec.cpp
--- a/Engine/Source/Runtime/AudioExtensions/Private/IAudioCodec.cpp
+++ b/Engine/Source/Runtime/AudioExtensions/Private/IAudioCodec.cpp
@@ -7,6 +7,8 @@
 
 namespace Audio
 {
+	PRAGMA_DISABLE_DEPRECATION_WARNINGS
+
 	// If < C++17, constexpr requires storage for linkage on Clang.
 	constexpr FFcc FFormatDescriptorSection::kSectionName;
 	constexpr FFcc FHeaderSection::kName;
@@ -14,7 +16,7 @@ namespace Audio
 	
 	// Encoder INPUT object. (simple for now).
 	template<typename TSampleType>
-	struct TEncoderInputObject : public IEncoderInput
+	struct UE_DEPRECATED(5.3, "Please use IAudioFormat/IStreamedCompressedInfo instead") TEncoderInputObject : public IEncoderInput
 	{
 		using FSampleType = TSampleType;
 		TSampleBuffer<FSampleType> Samples;
@@ -60,7 +62,7 @@ namespace Audio
 		return MakeUnique<TEncoderInputObject<int16>>(In16bitInterleave, InFmt);	
 	}
 		
-	struct FEncoderOutputBase : public IEncoderOutput
+	struct UE_DEPRECATED(5.3, "Please use IAudioFormat/IStreamedCompressedInfo instead") FEncoderOutputBase : public IEncoderOutput
 	{
 		TUniquePtr<FArchive> Archive;
 		TUniquePtr<FNestedSection> HeaderSection;
@@ -165,17 +167,14 @@ namespace Audio
 	IEncoderOutput::Create(
 		TArray<uint8>& OutExternalBuffer)
 	{
-		return MakeUnique<FEncoderOutputBase>(MakeUnique<FMemoryWriter>(OutExternalBuffer));
+		checkNoEntry();
+		return {};
 	}
 
 	Audio::IEncoderOutput::FPtr IEncoderOutput::Create()
 	{
-		struct FArrayOutputObject : public FEncoderOutputBase
-		{
-			TArray<uint8> Buffer;
-			FArrayOutputObject() : FEncoderOutputBase(MakeUnique<FMemoryWriter>(Buffer)) {}
-		};
-		return MakeUnique<FArrayOutputObject>();
+		checkNoEntry();
+		return {};
 	}
 
 	TUniquePtr<IDecoderInput> IDecoderInput::Create(
@@ -277,8 +276,7 @@ namespace Audio
 	TArrayView<const uint8> FDecoderInputBase::PeekNextPacket(
 		int32 InMaxPacketLength) const
 	{		
-		// TODO.
-		check(false);
+		checkNoEntry();
 		return MakeArrayView(static_cast<const uint8*>(0), 0);
 	}
 
@@ -425,15 +423,13 @@ namespace Audio
 	{
 		if (Requirements.DownstreamFormat == EBitRepresentation::Float32_Interleaved)
 		{
-			return MakeUnique<TCircularOutputBuffer<float>>(Requirements);
-			//return MakeUnique<TDecoderOutputOwnBuffer<float>>(Requirements);
+			return MakeUnique<TCircularOutputBuffer<float>>(Requirements);			
 		}
 		else if (Requirements.DownstreamFormat == EBitRepresentation::Int16_Interleaved)
 		{
 			return MakeUnique<TCircularOutputBuffer<int16>>(Requirements);
-			//return MakeUnique<TDecoderOutputOwnBuffer<int16>>(Requirements);
 		}
-		checkNoEntry(); //-V779
+		checkNoEntry(); 
 		return nullptr;
 	}
 
@@ -453,7 +449,7 @@ namespace Audio
 			Output->SetArrayViewBytes(InExternalMemory);
 			return Output;
 		}
-		checkNoEntry(); //-V779
+		checkNoEntry(); 
 		return nullptr;
 	}
 
@@ -468,5 +464,7 @@ namespace Audio
 		}
 		return false;
 	}
+
+	PRAGMA_ENABLE_DEPRECATION_WARNINGS
 }
 
diff --git a/Engine/Source/Runtime/AudioExtensions/Private/IAudioCodecRegistry.cpp b/Engine/Source/Runtime/AudioExtensions/Private/IAudioCodecRegistry.cpp
--- a/Engine/Source/Runtime/AudioExtensions/Private/IAudioCodecRegistry.cpp
+++ b/Engine/Source/Runtime/AudioExtensions/Private/IAudioCodecRegistry.cpp
@@ -16,6 +16,8 @@ DEFINE_LOG_CATEGORY(LogAudioCodec);
 
 namespace Audio
 {
+	PRAGMA_DISABLE_DEPRECATION_WARNINGS
+
 	class FCodecRegistry : public ICodecRegistry
 	{		
 		// Ownership.
@@ -28,13 +30,6 @@ namespace Audio
 		// RW lock on above containers.
 		mutable FRWLock ContainersLock;
 
-		//TMap<FName, TArray<FCodecPtr>> FamilyLookup;	// FamilyName -> [ICodecs]
-
-		// On Editor all other platform codecs with be discovered here
-		// Version|FamilyName|ImplName|		Flags| Features
-		// [0]    [Opus]     [OpusHwXbox]	[D]   [HwDecoder][HwResampler]
-		// [0]    [Opus]     [Opus]			[D+E]
-
 		// Anything matching the same version and family will be compatible.
 		// This should mean that we could test hardware/software implementations against each other.				
 
@@ -250,24 +245,8 @@ namespace Audio
 		{}
 
 		FCodecRegistry()
-		{
-			/*IModularFeatures::Get().OnModularFeatureRegistered().AddLambda([this](const FName& Type, IModularFeature* ModularFeature) 
-			{
-				if (Type == ICodec::GetModularFeatureName() && ModularFeature)
-				{
-					RegisterCodec(FCodecPtr(ModularFeature));
-				}
-			});
-
-			IModularFeatures::Get().OnModularFeatureUnregistered().AddLambda([this](const FName& Type, IModularFeature* ModularFeature) 
-			{
-				if (Type == ICodec::GetModularFeatureName() && ModularFeature)
-				{
-					UnregisterCodec(FCodecPtr(ModularFeature));
-				}
-			});			*/
+		{			
 		}
-
 	};
 
 	// Static instance.
@@ -283,4 +262,6 @@ namespace Audio
 		return *Instance;
 	}
 
+	PRAGMA_ENABLE_DEPRECATION_WARNINGS
+
 } //namespace Audio
diff --git a/Engine/Source/Runtime/AudioExtensions/Private/IAudioDecoderScheduler.cpp b/Engine/Source/Runtime/AudioExtensions/Private/IAudioDecoderScheduler.cpp
deleted file mode 100644
diff --git a/Engine/Source/Runtime/AudioExtensions/Private/PcmCodec.cpp b/Engine/Source/Runtime/AudioExtensions/Private/PcmCodec.cpp
deleted file mode 100644
diff --git a/Engine/Source/Runtime/AudioExtensions/Private/PcmCodec.h b/Engine/Source/Runtime/AudioExtensions/Private/PcmCodec.h
deleted file mode 100644
diff --git a/Engine/Source/Runtime/AudioExtensions/Public/IAudioCodec.h b/Engine/Source/Runtime/AudioExtensions/Public/IAudioCodec.h
--- a/Engine/Source/Runtime/AudioExtensions/Public/IAudioCodec.h
+++ b/Engine/Source/Runtime/AudioExtensions/Public/IAudioCodec.h
@@ -31,7 +31,7 @@
 #include "UObject/NameTypes.h"
 #include "UObject/Object.h"
 #include "UObject/ObjectMacros.h"
-#include "UObject/UObjectGlobals.h"
+#include "UObject/UObjectGlobals.h" 
 #include "UObject/UnrealNames.h"
 
 #include <initializer_list>
@@ -61,7 +61,7 @@ namespace Audio
 	struct FMixerSourceVoiceBuffer;
 
 	// Four character code class.
-	class FFcc
+	class UE_DEPRECATED(5.3, "Please use IAudioFormat/IStreamedCompressedInfo instead") FFcc
 	{
 	public:
 		FFcc() = default;
@@ -113,24 +113,27 @@ namespace Audio
 			static_assert(sizeof(From) == 0, "Unsupported conversion");
 		}
 	}
-		
-	enum EBitRepresentation
+
+	PRAGMA_DISABLE_DEPRECATION_WARNINGS
+	
+	enum UE_DEPRECATED(5.3, "Please use IAudioFormat/IStreamedCompressedInfo instead") EBitRepresentation
 	{
 		Int16_Interleaved,
 		Float32_Interleaved
 	};
 
-	struct AUDIOEXTENSIONS_API FDecodedFormatInfo 
+	struct AUDIOEXTENSIONS_API UE_DEPRECATED(5.3, "Please use IAudioFormat/IStreamedCompressedInfo instead") FDecodedFormatInfo
 	{		
 		uint32 NumChannels = 0;
 		uint32 NumFramesPerSec = 0;
+
 		EBitRepresentation BitRep = Int16_Interleaved;
 
 		int32 LoopStartOffset = INDEX_NONE;	// Where the loop start point is in the entire stream
 		int32 LoopEndOffset = INDEX_NONE;	// Where the loop end point is in the entire stream
 	};
-	
-	struct AUDIOEXTENSIONS_API IDecoderInput
+			
+	struct AUDIOEXTENSIONS_API UE_DEPRECATED(5.3, "Please use IAudioFormat/IStreamedCompressedInfo instead") IDecoderInput
 	{
 		// Factory.
 		static TUniquePtr<IDecoderInput> Create(
@@ -154,8 +157,8 @@ namespace Audio
 		virtual TArrayView<const uint8> PopNextPacket(
 			int32 InPacketSize ) = 0;
 	};
-	
-	struct AUDIOEXTENSIONS_API IDecoderOutput
+		
+	struct AUDIOEXTENSIONS_API UE_DEPRECATED(5.3, "Please use IAudioFormat/IStreamedCompressedInfo instead") IDecoderOutput
 	{	
 		struct FRequirements
 		{
@@ -234,12 +237,10 @@ namespace Audio
 		{
 			return PopAudio(MakeArrayView(InExternalBuffer.GetData(), InExternalBuffer.Num()), OutDetails);
 		}
-	};
-		
-
-	// Decode input wrappers.
+	};		
 
-	struct AUDIOEXTENSIONS_API FCodecFeatures
+	// Decode input wrappers.	
+	struct AUDIOEXTENSIONS_API UE_DEPRECATED(5.3, "Please use IAudioFormat/IStreamedCompressedInfo instead") FCodecFeatures
 	{	
 		enum EFeatureBitField
 		{
@@ -266,8 +267,8 @@ namespace Audio
 			}
 		}
 	};
-
-	struct AUDIOEXTENSIONS_API IDecoder
+	
+	struct AUDIOEXTENSIONS_API UE_DEPRECATED(5.3, "Please use IAudioFormat/IStreamedCompressedInfo instead") IDecoder
 	{		
 		enum EDecodeResult
 		{
@@ -288,8 +289,8 @@ namespace Audio
 		// Do the decode.
 		virtual FDecodeReturn Decode(bool bIsLooping = false) = 0;	
 	};
-
-	struct AUDIOEXTENSIONS_API FEncodedSectionBase
+	
+	struct AUDIOEXTENSIONS_API UE_DEPRECATED(5.3, "Please use IAudioFormat/IStreamedCompressedInfo instead") FEncodedSectionBase
 	{
 		// Serialized members.
 		FFcc SectionName;
@@ -310,8 +311,8 @@ namespace Audio
 		virtual bool EndSection(FArchive& Ar);
 		virtual void Serialize(FArchive& Ar) = 0;
 	};		
-
-	struct AUDIOEXTENSIONS_API FFormatDescriptorSection : public FEncodedSectionBase
+		
+	struct AUDIOEXTENSIONS_API UE_DEPRECATED(5.3, "Please use IAudioFormat/IStreamedCompressedInfo instead") FFormatDescriptorSection : public FEncodedSectionBase
 	{
 		static constexpr FFcc kSectionName = { 'D','e','s','c' };
 		static const uint32 kSectionVer = 2;
@@ -353,27 +354,28 @@ namespace Audio
 		
 		void Serialize(FArchive& Ar) override;
 	};
-
-	struct AUDIOEXTENSIONS_API FNestedSection : public FEncodedSectionBase
+	
+	struct AUDIOEXTENSIONS_API UE_DEPRECATED(5.3, "Please use IAudioFormat/IStreamedCompressedInfo instead") FNestedSection : public FEncodedSectionBase
 	{
 		FNestedSection(FFcc InName = FFcc())  : FEncodedSectionBase(InName, 0) {}
 		void Serialize(FArchive& Ar) override {}
 	};
-
-	struct AUDIOEXTENSIONS_API FHeaderSection : public FNestedSection
+	
+	struct AUDIOEXTENSIONS_API UE_DEPRECATED(5.3, "Please use IAudioFormat/IStreamedCompressedInfo instead") FHeaderSection : public FNestedSection
 	{
 		static constexpr FFcc kName = {'H','d','r','!' };
 		FHeaderSection() : FNestedSection(kName) {}
 		void Serialize(FArchive& Ar) override {}
 	};
-	struct AUDIOEXTENSIONS_API FSampleSection : public FNestedSection
+	
+	struct AUDIOEXTENSIONS_API UE_DEPRECATED(5.3, "Please use IAudioFormat/IStreamedCompressedInfo instead") FSampleSection : public FNestedSection
 	{
 		static constexpr FFcc kName = {'S','m','p','l'};
 		FSampleSection() : FNestedSection(kName) {}
 		void Serialize(FArchive& Ar) override {}
 	};
 	
-	struct AUDIOEXTENSIONS_API FRawSection : public FEncodedSectionBase
+	struct AUDIOEXTENSIONS_API UE_DEPRECATED(5.3, "Please use IAudioFormat/IStreamedCompressedInfo instead") FRawSection : public FEncodedSectionBase
 	{
 		TArrayView<const uint8> Buffer;
 		FRawSection(FFcc InName, TArrayView<const uint8> InBuffer = TArrayView<const uint8>() )
@@ -383,8 +385,8 @@ namespace Audio
 
 		void Serialize(FArchive& Ar) override {}
 	};
-
-	class AUDIOEXTENSIONS_API FDecoderInputBase : public Audio::IDecoderInput
+	
+	class AUDIOEXTENSIONS_API UE_DEPRECATED(5.3, "Please use IAudioFormat/IStreamedCompressedInfo instead") FDecoderInputBase : public Audio::IDecoderInput
 	{
 	protected:
 		FArchive* Archive = nullptr;
@@ -426,8 +428,8 @@ namespace Audio
 		// We need this for back compatibility.
 		int64 Tell() const override;
 	};
-
-	struct AUDIOEXTENSIONS_API FDecoderInputArrayView : public FDecoderInputBase
+	
+	struct AUDIOEXTENSIONS_API UE_DEPRECATED(5.3, "Please use IAudioFormat/IStreamedCompressedInfo instead") FDecoderInputArrayView : public FDecoderInputBase
 	{
 		TUniquePtr<FArchive> Reader;
 
@@ -456,9 +458,9 @@ namespace Audio
 			SetArchive(Reader.Get(), InOffsetInStream);
 		}
 	};
-		
+			
 	template<typename TSampleType>
-	struct TDecoderOutputArrayView : public IDecoderOutput
+	struct UE_DEPRECATED(5.3, "Please use IAudioFormat/IStreamedCompressedInfo instead") TDecoderOutputArrayView : public IDecoderOutput
 	{
 		FRequirements Reqs;
 		TArrayView<TSampleType> Buffer;
@@ -557,9 +559,9 @@ namespace Audio
 			return PopAudioInternal(InExternalInt16Buffer);
 		}	
 	};
-
+	
 	template<typename T>
-	struct TDecoderOutputOwnBuffer : public TDecoderOutputArrayView<T>
+	struct UE_DEPRECATED(5.3, "Please use IAudioFormat/IStreamedCompressedInfo instead") TDecoderOutputOwnBuffer : public TDecoderOutputArrayView<T>
 	{
 		TArray<T, TAlignedHeapAllocator<SIMD_ALIGNMENT>> AlignedMemory;
 		TDecoderOutputOwnBuffer(
@@ -570,9 +572,9 @@ namespace Audio
 			TDecoderOutputArrayView<T>::SetArrayView(MakeArrayView(AlignedMemory));
 		}
 	};
-
+	
 	template<typename TSampleType>
-	class TCircularOutputBuffer : public Audio::IDecoderOutput
+	class UE_DEPRECATED(5.3, "Please use IAudioFormat/IStreamedCompressedInfo instead") TCircularOutputBuffer : public Audio::IDecoderOutput
 	{
 		Audio::TCircularAudioBuffer<TSampleType> Buffer;
 		FRequirements Reqs;
@@ -649,7 +651,7 @@ namespace Audio
 #pragma region Encoder
 
 UCLASS(Abstract)
-class AUDIOEXTENSIONS_API UAudioCodecEncoderSettings : public UObject
+class UE_DEPRECATED(5.3, "Please use IAudioFormat/IStreamedCompressedInfo instead") AUDIOEXTENSIONS_API  UAudioCodecEncoderSettings : public UObject
 {
 public:
 	GENERATED_BODY()
@@ -663,8 +665,8 @@ public:
 };	
 
 namespace Audio
-{	
-	class AUDIOEXTENSIONS_API IEncoderOutput
+{		
+	class AUDIOEXTENSIONS_API UE_DEPRECATED(5.3, "Please use IAudioFormat/IStreamedCompressedInfo instead") IEncoderOutput
 	{
 	public:
 		using FPtr = TUniquePtr<IEncoderOutput>;
@@ -687,8 +689,8 @@ namespace Audio
 		virtual bool BeginSampleData() = 0;
 		virtual bool EndSampleData() = 0;
 	};
-
-	class AUDIOEXTENSIONS_API IEncoderInput
+	
+	class AUDIOEXTENSIONS_API UE_DEPRECATED(5.3, "Please use IAudioFormat/IStreamedCompressedInfo instead") IEncoderInput
 	{
 	public:
 		using FFormat = FDecodedFormatInfo;
@@ -711,8 +713,8 @@ namespace Audio
 		virtual const TSampleBuffer<int16>& 
 			GetSamples() const = 0;
 	};
-
-	struct AUDIOEXTENSIONS_API IEncoder
+	
+	struct AUDIOEXTENSIONS_API UE_DEPRECATED(5.3, "Please use IAudioFormat/IStreamedCompressedInfo instead") IEncoder
 	{
 		virtual ~IEncoder() {}
 
@@ -744,8 +746,8 @@ namespace Audio
 #pragma region Codec
 
 namespace Audio
-{
-	struct AUDIOEXTENSIONS_API FCodecDetails
+{	
+	struct AUDIOEXTENSIONS_API UE_DEPRECATED(5.3, "Please use IAudioFormat/IStreamedCompressedInfo instead") FCodecDetails
 	{
 		FName Name;								// Unique name for this implementation. e.g. "OpusXbox"
 		FName FamilyName;						// Family name e.g. "Opus". Same as old Format name.	
@@ -765,8 +767,8 @@ namespace Audio
 			: Name(InName), FamilyName(InFamilyName), Version(InVersion), Features(InFeatures) 
 		{}
 	};
-
-	struct AUDIOEXTENSIONS_API ICodec 
+	
+	struct AUDIOEXTENSIONS_API UE_DEPRECATED(5.3, "Please use IAudioFormat/IStreamedCompressedInfo instead") ICodec
 	{
 		virtual ~ICodec() {}
 
@@ -792,6 +794,9 @@ namespace Audio
 			return nullptr; 
 		}			
 	};
+
+	PRAGMA_ENABLE_DEPRECATION_WARNINGS
+
 } // namespace Audio
 
 #pragma endregion Codec
diff --git a/Engine/Source/Runtime/AudioExtensions/Public/IAudioCodecRegistry.h b/Engine/Source/Runtime/AudioExtensions/Public/IAudioCodecRegistry.h
--- a/Engine/Source/Runtime/AudioExtensions/Public/IAudioCodecRegistry.h
+++ b/Engine/Source/Runtime/AudioExtensions/Public/IAudioCodecRegistry.h
@@ -15,7 +15,9 @@ namespace Audio
 	struct ICodec;
 	struct IDecoderInput;
 
-	class AUDIOEXTENSIONS_API ICodecRegistry
+	PRAGMA_DISABLE_DEPRECATION_WARNINGS
+	
+	class AUDIOEXTENSIONS_API UE_DEPRECATED(5.3, "ICodecRegistry is deprecated, please use IAudioInfoFactoryRegistry for decoders and IAudioFormat for encoders") ICodecRegistry
 	{
 		static TUniquePtr<ICodecRegistry> Instance;
 	public:
@@ -42,4 +44,6 @@ namespace Audio
 		// Find default codec for a platform. (none uses host platform).
 		virtual FCodecPtr FindDefaultCodec(FName InPlatformName = NAME_None) const = 0;
 	};
-}
+
+	PRAGMA_ENABLE_DEPRECATION_WARNINGS
+}
\ No newline at end of file
diff --git a/Engine/Source/Runtime/AudioExtensions/Public/IAudioDecoderScheduler.h b/Engine/Source/Runtime/AudioExtensions/Public/IAudioDecoderScheduler.h
deleted file mode 100644
