Fix for potential memory corruption

diff --git a/Engine/Source/Runtime/AudioExtensions/Private/IAudioCodec.cpp b/Engine/Source/Runtime/AudioExtensions/Private/IAudioCodec.cpp
--- a/Engine/Source/Runtime/AudioExtensions/Private/IAudioCodec.cpp
+++ b/Engine/Source/Runtime/AudioExtensions/Private/IAudioCodec.cpp
@@ -283,29 +283,8 @@ namespace Audio
 	TArrayView<const uint8> FDecoderInputBase::PopNextPacket(
 		int32 InPacketSize)
 	{	
-		if (PopBuffer.Num() < InPacketSize)
-		{
-			PopBuffer.SetNumZeroed(InPacketSize);
-		}		
-		if (!bSampleSection )
-		{
-			// To start serving packets, we must be in the sample section.
-			if( !audio_ensure(BeginSampleData()))
-			{
-				// Fail.
-				return MakeArrayView(static_cast<const uint8*>(0), 0);
-			}
-		}
-		if (!audio_ensure(!Archive->IsError()))
-		{
-			return MakeArrayView(static_cast<const uint8*>(0), 0);
-		}
-		Archive->Serialize(PopBuffer.GetData(), InPacketSize);
-		if (!audio_ensure(!Archive->IsError()))
-		{
-			return MakeArrayView(static_cast<const uint8*>(0), 0);
-		}
-		return MakeArrayView(PopBuffer.GetData(),InPacketSize);
+		checkNoEntry();
+		return {};
 	}
 	
 	bool FDecoderInputBase::BeginSampleData() 
