Fix submix failing postload check

diff --git a/Engine/Source/Runtime/AudioExtensions/Private/AudioExtensionsModule.cpp b/Engine/Source/Runtime/AudioExtensions/Private/AudioExtensionsModule.cpp
--- a/Engine/Source/Runtime/AudioExtensions/Private/AudioExtensionsModule.cpp
+++ b/Engine/Source/Runtime/AudioExtensions/Private/AudioExtensionsModule.cpp
@@ -3,22 +3,41 @@
 #include "Modules/ModuleManager.h"
 
 #include "AudioExtensionsLog.h"
+#include "AudioExtentionsModule.h"
 #include "HAL/LowLevelMemStats.h"
 #include "IAudioExtensionPlugin.h"
+#include "Algo/ForEach.h"
+#include "Misc/DataDrivenPlatformInfoRegistry.h"
 
 DECLARE_LLM_MEMORY_STAT(TEXT("AudioSpatializationPlugins"), STAT_AudioSpatializationPluginsLLM, STATGROUP_LLMFULL);
 LLM_DEFINE_TAG(Audio_SpatializationPlugins, NAME_None, TEXT("Audio"), GET_STATFNAME(STAT_AudioSpatializationPluginsLLM), GET_STATFNAME(STAT_AudioSummaryLLM));
 
 DEFINE_LOG_CATEGORY(LogAudioExtensions)
 
-class FAudioExtensionsModule : public IModuleInterface
+FAudioExtensionsModule* FAudioExtensionsModule::Get()
 {
-public:
-	void StartupModule() override
+	// This a fast check if loaded and will load the module if not.
+	return FModuleManager::Get().LoadModulePtr<FAudioExtensionsModule>(TEXT("AudioExtensions"));
+}
+
+void FAudioExtensionsModule::StartupModule()
+{
+	// Load Platform specific Audio feature modules. 
+	TArray<FName> Platforms = FDataDrivenPlatformInfoRegistry::GetSortedPlatformNames(EPlatformInfoType::TruePlatformsOnly);
+	Algo::ForEach(Platforms,[](const FName InName)
 	{
-		FModuleManager::Get().LoadModuleChecked(TEXT("SignalProcessing"));
-		FModuleManager::Get().LoadModuleChecked(TEXT("AudioMixerCore"));
-	}
-};
+		// Load any modules that are named <PlatformName>AudioFeatures, which will load hidden platforms.
+		// This equivalent, but slightly safer that using the wildcard "*AudioFeatures" but without the risk
+		// of accidentally bringing in user modules.
+		const FString PlatformModuleName = InName.ToString() + TEXT("AudioFeatures");
+
+		// Most of these calls will fail, but will do so quickly as the module manager already knows about all
+		// possible modules, so this is a quick lookup.
+		FModuleManager::Get().LoadModule(*PlatformModuleName);
+	});
+		
+	FModuleManager::Get().LoadModuleChecked(TEXT("SignalProcessing"));
+	FModuleManager::Get().LoadModuleChecked(TEXT("AudioMixerCore"));
+}
 
 IMPLEMENT_MODULE(FAudioExtensionsModule, AudioExtensions);
diff --git a/Engine/Source/Runtime/AudioExtensions/Private/AudioExtentionsModule.h b/Engine/Source/Runtime/AudioExtensions/Private/AudioExtentionsModule.h
new file mode 100644
--- /dev/null
+++ b/Engine/Source/Runtime/AudioExtensions/Private/AudioExtentionsModule.h
@@ -0,0 +1,12 @@
+// Copyright Epic Games, Inc. All Rights Reserved.
+#pragma once
+
+#include "AudioExtentionsModule.h"
+
+class FAudioExtensionsModule final : public IModuleInterface
+{
+public:
+	static FAudioExtensionsModule* Get(); 
+	
+	virtual void StartupModule() override;
+};
\ No newline at end of file
diff --git a/Engine/Source/Runtime/AudioExtensions/Private/IAudioEndpoint.cpp b/Engine/Source/Runtime/AudioExtensions/Private/IAudioEndpoint.cpp
--- a/Engine/Source/Runtime/AudioExtensions/Private/IAudioEndpoint.cpp
+++ b/Engine/Source/Runtime/AudioExtensions/Private/IAudioEndpoint.cpp
@@ -2,6 +2,8 @@
 
 #include "IAudioEndpoint.h"
 
+#include "AudioExtentionsModule.h"
+
 #include UE_INLINE_GENERATED_CPP_BY_NAME(IAudioEndpoint)
 
 DEFINE_LOG_CATEGORY(LogAudioEndpoints);
@@ -180,6 +182,9 @@ IAudioEndpointFactory* IAudioEndpointFactory::Get(const FName& InName)
 
 TArray<FName> IAudioEndpointFactory::GetAvailableEndpointTypes()
 {
+	// Ensure the module is loaded. This will cause any platform extension modules to load and register.
+	ensure(FAudioExtensionsModule::Get() != nullptr);
+	
 	TArray<FName> EndpointNames;
 
 	EndpointNames.Add(GetTypeNameForDefaultEndpoint());
diff --git a/Engine/Source/Runtime/AudioExtensions/Private/ISoundfieldEndpoint.cpp b/Engine/Source/Runtime/AudioExtensions/Private/ISoundfieldEndpoint.cpp
--- a/Engine/Source/Runtime/AudioExtensions/Private/ISoundfieldEndpoint.cpp
+++ b/Engine/Source/Runtime/AudioExtensions/Private/ISoundfieldEndpoint.cpp
@@ -2,6 +2,8 @@
 
 #include "ISoundfieldEndpoint.h"
 
+#include "AudioExtentionsModule.h"
+
 #include UE_INLINE_GENERATED_CPP_BY_NAME(ISoundfieldEndpoint)
 
 ISoundfieldEndpoint::ISoundfieldEndpoint(int32 NumRenderCallbacksToBuffer)
@@ -136,6 +138,9 @@ ISoundfieldEndpointFactory* ISoundfieldEndpointFactory::Get(const FName& InName)
 
 TArray<FName> ISoundfieldEndpointFactory::GetAllSoundfieldEndpointTypes()
 {
+	// Ensure the module is loaded. This will cause any platform extension modules to load and register. 
+	ensure(FAudioExtensionsModule::Get() != nullptr);
+	
 	TArray<FName> SoundfieldFormatNames;
 
 	SoundfieldFormatNames.Add(DefaultSoundfieldEndpointName());
diff --git a/Engine/Source/Runtime/AudioExtensions/Private/ISoundfieldFormat.cpp b/Engine/Source/Runtime/AudioExtensions/Private/ISoundfieldFormat.cpp
--- a/Engine/Source/Runtime/AudioExtensions/Private/ISoundfieldFormat.cpp
+++ b/Engine/Source/Runtime/AudioExtensions/Private/ISoundfieldFormat.cpp
@@ -1,6 +1,8 @@
 // Copyright Epic Games, Inc. All Rights Reserved.
 
 #include "ISoundfieldFormat.h"
+
+#include "AudioExtentionsModule.h"
 #include "ISoundfieldEndpoint.h"
 
 #include UE_INLINE_GENERATED_CPP_BY_NAME(ISoundfieldFormat)
@@ -65,6 +67,9 @@ ISoundfieldFactory* ISoundfieldFactory::Get(const FName& InName)
 
 TArray<FName> ISoundfieldFactory::GetAvailableSoundfieldFormats()
 {
+	// Ensure the module is loaded. This will cause any platform extension modules to load and register. 
+	ensure(FAudioExtensionsModule::Get() != nullptr);
+	
 	TArray<FName> SoundfieldFormatNames;
 
 	SoundfieldFormatNames.Add(GetFormatNameForInheritedEncoding());
