Remove self-included header and replace with needed one for non-unity
#rnx

diff --git a/Engine/Source/Runtime/AudioExtensions/Private/AudioExtentionsModule.h b/Engine/Source/Runtime/AudioExtensions/Private/AudioExtentionsModule.h
--- a/Engine/Source/Runtime/AudioExtensions/Private/AudioExtentionsModule.h
+++ b/Engine/Source/Runtime/AudioExtensions/Private/AudioExtentionsModule.h
@@ -1,7 +1,7 @@
 // Copyright Epic Games, Inc. All Rights Reserved.
 #pragma once
 
-#include "AudioExtentionsModule.h"
+#include "Modules/ModuleInterface.h"
 
 class FAudioExtensionsModule final : public IModuleInterface
 {
