[AudioProperties] Set up modules and base set of classes with support for basic inheritance and injection
#rnx

diff --git a/Engine/Source/Runtime/AudioExtensions/Public/AudioPropertiesSheetAssetBase.h b/Engine/Source/Runtime/AudioExtensions/Public/AudioPropertiesSheetAssetBase.h
new file mode 100644
--- /dev/null
+++ b/Engine/Source/Runtime/AudioExtensions/Public/AudioPropertiesSheetAssetBase.h
@@ -0,0 +1,18 @@
+// Copyright Epic Games, Inc. All Rights Reserved.
+
+#pragma once
+
+#include "UObject/Object.h"
+#include "UObject/ObjectPtr.h"
+
+#include "AudioPropertiesSheetAssetBase.generated.h"
+
+
+UCLASS(Abstract)
+class AUDIOEXTENSIONS_API UAudioPropertiesSheetAssetBase : public UObject
+{
+	GENERATED_BODY()
+
+public: 
+	virtual bool CopyToObjectProperties(TObjectPtr<UObject> TargetObject) const { return false; };
+};
\ No newline at end of file
