[AudioProperties] Create a 'Bindings' asset type to determine which properties should be injected from a sheet to a target uobject. This replaces the use of uproperty tag names

#rnx

diff --git a/Engine/Source/Runtime/AudioExtensions/Public/AudioPropertiesBindings.h b/Engine/Source/Runtime/AudioExtensions/Public/AudioPropertiesBindings.h
new file mode 100644
--- /dev/null
+++ b/Engine/Source/Runtime/AudioExtensions/Public/AudioPropertiesBindings.h
@@ -0,0 +1,19 @@
+// Copyright Epic Games, Inc. All Rights Reserved.
+
+#pragma once
+
+#include "UObject/Object.h"
+
+#include "AudioPropertiesBindings.generated.h"
+
+UCLASS()
+class AUDIOEXTENSIONS_API UAudioPropertiesBindings : public UObject
+{
+	GENERATED_BODY()
+
+	UAudioPropertiesBindings() = default;
+
+public:
+	UPROPERTY(EditDefaultsOnly, Category = "Default")
+	TMap<FName, FName> ObjectPropertyToSheetPropertyMap;
+};
\ No newline at end of file
diff --git a/Engine/Source/Runtime/AudioExtensions/Public/AudioPropertiesSheetAssetBase.h b/Engine/Source/Runtime/AudioExtensions/Public/AudioPropertiesSheetAssetBase.h
--- a/Engine/Source/Runtime/AudioExtensions/Public/AudioPropertiesSheetAssetBase.h
+++ b/Engine/Source/Runtime/AudioExtensions/Public/AudioPropertiesSheetAssetBase.h
@@ -7,6 +7,7 @@
 
 #include "AudioPropertiesSheetAssetBase.generated.h"
 
+class UAudioPropertiesBindings;
 
 UCLASS(Abstract)
 class AUDIOEXTENSIONS_API UAudioPropertiesSheetAssetBase : public UObject
@@ -14,5 +15,5 @@ class AUDIOEXTENSIONS_API UAudioPropertiesSheetAssetBase : public UObject
 	GENERATED_BODY()
 
 public: 
-	virtual bool CopyToObjectProperties(TObjectPtr<UObject> TargetObject) const { return false; };
+	virtual bool CopyToObjectProperties(TObjectPtr<UObject> TargetObject, const TObjectPtr<const UAudioPropertiesBindings> PropertiesBindings) const { return false; };
 };
\ No newline at end of file
diff --git a/Engine/Source/Runtime/AudioExtensions/Public/IAudioPropertiesSheet.h b/Engine/Source/Runtime/AudioExtensions/Public/IAudioPropertiesSheet.h
new file mode 100644
--- /dev/null
+++ b/Engine/Source/Runtime/AudioExtensions/Public/IAudioPropertiesSheet.h
@@ -0,0 +1,18 @@
+// Copyright Epic Games, Inc. All Rights Reserved.
+
+#pragma once
+
+#include "UObject/Object.h"
+#include "UObject/ObjectPtr.h"
+
+#include "IAudioPropertiesSheet.generated.h"
+
+
+UCLASS(Abstract)
+class AUDIOEXTENSIONS_API UAudioPropertySheetBaseAsset : public UObject
+{
+	GENERATED_BODY()
+
+public: 
+	virtual bool CopyToObjectProperties(TObjectPtr<UObject> TargetObject) { return false; };
+};
\ No newline at end of file
