Set Iwaveformtransformation and Waveform Transformations module  to editor only

diff --git a/Engine/Plugins/Editor/WaveformEditor/WaveformEditor.uplugin b/Engine/Plugins/Editor/WaveformEditor/WaveformEditor.uplugin
--- a/Engine/Plugins/Editor/WaveformEditor/WaveformEditor.uplugin
+++ b/Engine/Plugins/Editor/WaveformEditor/WaveformEditor.uplugin
@@ -27,7 +27,7 @@
 		},
 		{
 			"Name": "WaveformTransformations",
-			"Type": "UncookedOnly",
+			"Type": "Editor",
 			"LoadingPhase": "PreDefault"
 		},
 		{
diff --git a/Engine/Source/Runtime/AudioExtensions/Public/IWaveformTransformation.h b/Engine/Source/Runtime/AudioExtensions/Public/IWaveformTransformation.h
--- a/Engine/Source/Runtime/AudioExtensions/Public/IWaveformTransformation.h
+++ b/Engine/Source/Runtime/AudioExtensions/Public/IWaveformTransformation.h
@@ -67,6 +67,8 @@ class UWaveformTransformationBase : public UObject
 public:
 	virtual Audio::FTransformationPtr CreateTransformation() const { return nullptr; }
 	virtual void UpdateConfiguration(FWaveTransformUObjectConfiguration& InOutConfiguration) {};
+	
+	virtual bool IsEditorOnly() const override { return true; }
 };
 
 // Object that holds an ordered list of transformations to perform on a sound wave
@@ -79,5 +81,7 @@ public:
 	UPROPERTY(EditAnywhere, Instanced, Category = "Transformations")
 	TArray<TObjectPtr<UWaveformTransformationBase>> Transformations;
 
+	virtual bool IsEditorOnly() const override { return true; }
+	
 	AUDIOEXTENSIONS_API TArray<Audio::FTransformationPtr> CreateTransformations() const;
 };
