Isolating out some easy wins from deeper investigations into audio parameter optimizations. The biggest change is to the array parameter version of FAudioParameter::Merge where we'll now avoid sorting both arrays and just doing a binary search into the output array. This results in fewer shifting inserts as well, since the order of the output array is not required to be sorted.

Saves about 0.15ms with a dozen AudioComponentGroups updating parameters per frame.

diff --git a/Engine/Plugins/AudioGameplay/Source/AudioGameplay/Private/AudioComponentGroup.cpp b/Engine/Plugins/AudioGameplay/Source/AudioGameplay/Private/AudioComponentGroup.cpp
--- a/Engine/Plugins/AudioGameplay/Source/AudioGameplay/Private/AudioComponentGroup.cpp
+++ b/Engine/Plugins/AudioGameplay/Source/AudioGameplay/Private/AudioComponentGroup.cpp
@@ -583,8 +583,7 @@ void UAudioComponentGroup::UpdateComponentParameters()
     	});
 	}
 	
-	TArray<FAudioParameter> Values = ParamsToSet;
-	FAudioParameter::Merge(MoveTemp(Values), PersistentParams);
+	FAudioParameter::Merge(MoveTemp(ParamsToSet), PersistentParams);
 	ParamsToSet.Reset();
 }
 
diff --git a/Engine/Plugins/Runtime/Metasound/Source/MetasoundFrontend/Private/MetasoundParameterTransmitter.cpp b/Engine/Plugins/Runtime/Metasound/Source/MetasoundFrontend/Private/MetasoundParameterTransmitter.cpp
--- a/Engine/Plugins/Runtime/Metasound/Source/MetasoundFrontend/Private/MetasoundParameterTransmitter.cpp
+++ b/Engine/Plugins/Runtime/Metasound/Source/MetasoundFrontend/Private/MetasoundParameterTransmitter.cpp
@@ -272,23 +272,21 @@ namespace Metasound
 	{
 		bool bSuccess = true;
 
-		TArray<FAudioParameter> NewParams;
-		for (FAudioParameter& Param : InParameters)
+		const bool bAllowShrinking = false;
+		for (int32 ParamIndex = InParameters.Num() - 1; ParamIndex >= 0; --ParamIndex)
 		{
+			FAudioParameter& Param = InParameters[ParamIndex];
 			const FName ParamName = Param.ParamName;
-			if (SetParameterWithLiteral(ParamName, Frontend::ConvertParameterToLiteral(Param)))
-			{
-				NewParams.Add(MoveTemp(Param));
-			}
-			else
+			if (!SetParameterWithLiteral(ParamName, Frontend::ConvertParameterToLiteral(Param)))
 			{
+				InParameters.RemoveAtSwap(ParamIndex, 1, bAllowShrinking);
 				bSuccess = false;
 			}
 		}
 
-		if (!NewParams.IsEmpty())
+		if (!InParameters.IsEmpty())
 		{
-			bSuccess &= FParameterTransmitterBase::SetParameters(MoveTemp(NewParams));
+			bSuccess &= FParameterTransmitterBase::SetParameters(MoveTemp(InParameters));
 		}
 
 		InParameters.Reset();
diff --git a/Engine/Source/Runtime/AudioExtensions/Private/AudioParameter.cpp b/Engine/Source/Runtime/AudioExtensions/Private/AudioParameter.cpp
--- a/Engine/Source/Runtime/AudioExtensions/Private/AudioParameter.cpp
+++ b/Engine/Source/Runtime/AudioExtensions/Private/AudioParameter.cpp
@@ -1,8 +1,9 @@
 // Copyright Epic Games, Inc. All Rights Reserved.
 #include "AudioParameter.h"
 
-#include UE_INLINE_GENERATED_CPP_BY_NAME(AudioParameter)
+#include "Algo/BinarySearch.h"
 
+#include UE_INLINE_GENERATED_CPP_BY_NAME(AudioParameter)
 
 namespace Audio
 {
@@ -190,40 +191,27 @@ void FAudioParameter::Merge(TArray<FAudioParameter>&& InParams, TArray<FAudioPar
 
 	if (OutParams.IsEmpty())
 	{
-		OutParams.Append(MoveTemp(InParams));
+		OutParams = MoveTemp(InParams);
 		return;
 	}
 
 	auto SortParamsPredicate = [](const FAudioParameter& A, const FAudioParameter& B) { return A.ParamName.FastLess(B.ParamName); };
-
-	InParams.Sort(SortParamsPredicate);
 	OutParams.Sort(SortParamsPredicate);
 
-	for (int32 i = OutParams.Num() - 1; i >= 0; --i)
+	for (FAudioParameter& NewParam : InParams)
 	{
-		while (!InParams.IsEmpty())
+		const int32 ExistingElementIndex = Algo::LowerBound(OutParams, NewParam, SortParamsPredicate);
+		if (OutParams.IsValidIndex(ExistingElementIndex))
 		{
-			FAudioParameter& OutParam = OutParams[i];
-			// Break going through in params as there will be no more matches with this current out param 
-			if (InParams.Last().ParamName.FastLess(OutParam.ParamName))
-			{
-				break;
-			}
-
-			constexpr bool bAllowShrinking = false;
-			FAudioParameter NewParam = InParams.Pop(bAllowShrinking);
-			if (NewParam.ParamName == OutParam.ParamName)
+			FAudioParameter& ExistingParam = OutParams[ExistingElementIndex];
+			if (ExistingParam.ParamName == NewParam.ParamName)
 			{
-				OutParam.Merge(NewParam);
-			}
-			else
-			{
-				OutParams.Emplace(MoveTemp(NewParam));
+				ExistingParam.Merge(NewParam);
+				continue;
 			}
 		}
-	}
 
-	// Add the rest of the in params that have no matching out param 
-	OutParams.Append(MoveTemp(InParams));
+		OutParams.Emplace(MoveTemp(NewParam));
+	}
 }
 
diff --git a/Engine/Source/Runtime/Engine/Private/Audio/SoundParameterControllerInterface.cpp b/Engine/Source/Runtime/Engine/Private/Audio/SoundParameterControllerInterface.cpp
--- a/Engine/Source/Runtime/Engine/Private/Audio/SoundParameterControllerInterface.cpp
+++ b/Engine/Source/Runtime/Engine/Private/Audio/SoundParameterControllerInterface.cpp
@@ -165,12 +165,11 @@ void ISoundParameterControllerInterface::SetParameters(TArray<FAudioParameter>&&
 			if (!ParamsToSet.IsEmpty())
 			{
 				DECLARE_CYCLE_STAT(TEXT("FAudioThreadTask.SoundParameterControllerInterface.SetParameters"), STAT_AudioSetParameters, STATGROUP_AudioThreadCommands);
-				AudioDevice->SendCommandToActiveSounds(GetInstanceOwnerID(), [AudioDevice, Params = MoveTemp(ParamsToSet)](FActiveSound& ActiveSound)
+				AudioDevice->SendCommandToActiveSounds(GetInstanceOwnerID(), [AudioDevice, Params = MoveTemp(ParamsToSet)](FActiveSound& ActiveSound) mutable
 				{
 					if (Audio::IParameterTransmitter* Transmitter = ActiveSound.GetTransmitter())
 					{
-						TArray<FAudioParameter> TempParams = Params;
-						Transmitter->SetParameters(MoveTemp(TempParams));
+						Transmitter->SetParameters(MoveTemp(Params));
 					}
 				}, GET_STATID(STAT_AudioSetParameters));
 			}
