Fix thread safety issue in CreateModulatorHandleId

diff --git a/Engine/Source/Runtime/AudioExtensions/Private/IAudioModulation.cpp b/Engine/Source/Runtime/AudioExtensions/Private/IAudioModulation.cpp
--- a/Engine/Source/Runtime/AudioExtensions/Private/IAudioModulation.cpp
+++ b/Engine/Source/Runtime/AudioExtensions/Private/IAudioModulation.cpp
@@ -9,6 +9,11 @@
 
 namespace Audio
 {
+	namespace ModulationPrivate
+	{
+		static std::atomic<FModulatorHandleId> NextHandleId = INDEX_NONE;
+	}
+
 	namespace ModulationInterfacePrivate
 	{
 		class FModulationParameterRegistry
@@ -66,8 +71,7 @@ namespace Audio
 
 	FModulatorHandleId CreateModulatorHandleId()
 	{
-		static FModulatorHandleId NextHandleId = INDEX_NONE;
-		return ++NextHandleId;
+		return ++ModulationPrivate::NextHandleId;
 	}
 
 	FModulationParameter::FModulationParameter()
