Removed the FAudioParameter(FName, TArray<UObject*>&&), constructor as an upcoming change to TArray will warn that this move from TArray<UObject*> to TArray<TObjectPtr<UObject>> is actually copying.  The fix will be to change the constructor to take TArray<TObjectPtr<UObject>> instead.

diff --git a/Engine/Source/Runtime/AudioExtensions/Public/AudioParameter.h b/Engine/Source/Runtime/AudioExtensions/Public/AudioParameter.h
--- a/Engine/Source/Runtime/AudioExtensions/Public/AudioParameter.h
+++ b/Engine/Source/Runtime/AudioExtensions/Public/AudioParameter.h
@@ -184,12 +184,15 @@ struct FAudioParameter
 	{
 	}
 
+// We can't move TArray<UObject*> into a TArray<TObjectPtr<UObject>> since UE_OBJECT_PTR_GC_BARRIER was added
+#if 0
 	FAudioParameter(FName InName, TArray<UObject*>&& InValue)
 		: ParamName(InName)
 		, ArrayObjectParam(MoveTemp(InValue))
 		, ParamType(EAudioParameterType::ObjectArray)
 	{
 	}
+#endif
 
 	FAudioParameter(FName InName, const TArray<FString>& InValue)
 		: ParamName(InName)
