[Property Sheets] bind to parent post edit change so that classes can bind to only one asset to get post edit change updates
GetTargetedPropertyNames allows to get a list of targeted properties (needed for details view)
remove unused base asset in audio extensions

diff --git a/Engine/Source/Runtime/AudioExtensions/Public/AudioPropertiesSheetAssetBase.h b/Engine/Source/Runtime/AudioExtensions/Public/AudioPropertiesSheetAssetBase.h
--- a/Engine/Source/Runtime/AudioExtensions/Public/AudioPropertiesSheetAssetBase.h
+++ b/Engine/Source/Runtime/AudioExtensions/Public/AudioPropertiesSheetAssetBase.h
@@ -2,13 +2,12 @@
 
 #pragma once
 
+#include "Delegates/IDelegateInstance.h"
 #include "UObject/Object.h"
 #include "UObject/ObjectPtr.h"
 
 #include "AudioPropertiesSheetAssetBase.generated.h"
 
-class UAudioPropertiesBindings;
-
 UCLASS(Abstract)
 class AUDIOEXTENSIONS_API UAudioPropertiesSheetAssetBase : public UObject
 {
@@ -18,5 +17,8 @@ public:
 #if WITH_EDITOR
 	UFUNCTION(BlueprintCallable, BlueprintPure=false, Category="AudioProperties")
 	virtual bool CopyToObjectProperties(UObject* TargetObject) const { return false; };
+
+	virtual FDelegateHandle BindPropertiesCopyToSheetChanges(UObject* TargetObject) { return FDelegateHandle(); };
+	virtual void UnbindCopyFromPropertySheetChanges(UObject* ObjectToUnbind) {};
 #endif
 };
\ No newline at end of file
diff --git a/Engine/Source/Runtime/AudioExtensions/Public/IAudioPropertiesSheet.h b/Engine/Source/Runtime/AudioExtensions/Public/IAudioPropertiesSheet.h
deleted file mode 100644
