[AudioProperties] remove pragma once on cpp file to fix for cis error on android
#rnx

diff --git a/Engine/Source/Runtime/AudioExtensions/Private/IAudioPropertiesDetailsInjector.cpp b/Engine/Source/Runtime/AudioExtensions/Private/IAudioPropertiesDetailsInjector.cpp
--- a/Engine/Source/Runtime/AudioExtensions/Private/IAudioPropertiesDetailsInjector.cpp
+++ b/Engine/Source/Runtime/AudioExtensions/Private/IAudioPropertiesDetailsInjector.cpp
@@ -1,8 +1,5 @@
 // Copyright Epic Games, Inc. All Rights Reserved.
 
-#pragma once
-
 #include "IAudioPropertiesDetailsInjector.h"
 
-
 IAudioPropertiesDetailsInjector::~IAudioPropertiesDetailsInjector() {};
\ No newline at end of file
