Fix for crash setting audio parameters

diff --git a/Engine/Source/Runtime/AudioExtensions/Private/IAudioParameterTransmitter.cpp b/Engine/Source/Runtime/AudioExtensions/Private/IAudioParameterTransmitter.cpp
--- a/Engine/Source/Runtime/AudioExtensions/Private/IAudioParameterTransmitter.cpp
+++ b/Engine/Source/Runtime/AudioExtensions/Private/IAudioParameterTransmitter.cpp
@@ -22,6 +22,7 @@ namespace Audio
 
 	bool FParameterTransmitterBase::GetParameter(FName InName, FAudioParameter& OutValue) const
 	{
+		FScopeLock Lock(&AudioParameterCritSec);
 		if (const FAudioParameter* Param = FAudioParameter::FindParam(AudioParameters, InName))
 		{
 			OutValue = *Param;
@@ -33,6 +34,8 @@ namespace Audio
 
 	void FParameterTransmitterBase::ResetParameters() 
 	{
+		FScopeLock Lock(&AudioParameterCritSec);
+
 		AudioParameters.Reset();
 	}
 	
@@ -49,6 +52,8 @@ namespace Audio
 
 	bool FParameterTransmitterBase::SetParameters(TArray<FAudioParameter>&& InParameters)
 	{
+		FScopeLock Lock(&AudioParameterCritSec);
+
 		FAudioParameter::Merge(MoveTemp(InParameters), AudioParameters);
 		return true;
 	}
diff --git a/Engine/Source/Runtime/AudioExtensions/Public/IAudioParameterTransmitter.h b/Engine/Source/Runtime/AudioExtensions/Public/IAudioParameterTransmitter.h
--- a/Engine/Source/Runtime/AudioExtensions/Public/IAudioParameterTransmitter.h
+++ b/Engine/Source/Runtime/AudioExtensions/Public/IAudioParameterTransmitter.h
@@ -104,6 +104,7 @@ namespace Audio
 		AUDIOEXTENSIONS_API virtual void OnRealizeVirtualizedActiveSound(TArray<FAudioParameter>&& InParameters);
 
 	protected:
+		mutable FCriticalSection AudioParameterCritSec; 
 		TArray<FAudioParameter> AudioParameters;
 		bool bIsVirtualized;
 	};
diff --git a/Engine/Source/Runtime/Engine/Private/Audio/SoundParameterControllerInterface.cpp b/Engine/Source/Runtime/Engine/Private/Audio/SoundParameterControllerInterface.cpp
--- a/Engine/Source/Runtime/Engine/Private/Audio/SoundParameterControllerInterface.cpp
+++ b/Engine/Source/Runtime/Engine/Private/Audio/SoundParameterControllerInterface.cpp
@@ -169,7 +169,8 @@ void ISoundParameterControllerInterface::SetParameters(TArray<FAudioParameter>&&
 				{
 					if (Audio::IParameterTransmitter* Transmitter = ActiveSound.GetTransmitter())
 					{
-						Transmitter->SetParameters(MoveTemp(Params));
+						TArray<FAudioParameter> Temp = Params;
+						Transmitter->SetParameters(MoveTemp(Temp));
 					}
 				}, GET_STATID(STAT_AudioSetParameters));
 			}
