USoundWave hot reload for MetaSounds in PIE and editor audition.

diff --git a/Engine/Plugins/Runtime/Metasound/Source/MetasoundEngine/Private/MetasoundWavePlayerNode.cpp b/Engine/Plugins/Runtime/Metasound/Source/MetasoundEngine/Private/MetasoundWavePlayerNode.cpp
--- a/Engine/Plugins/Runtime/Metasound/Source/MetasoundEngine/Private/MetasoundWavePlayerNode.cpp
+++ b/Engine/Plugins/Runtime/Metasound/Source/MetasoundEngine/Private/MetasoundWavePlayerNode.cpp
@@ -524,6 +524,13 @@ namespace Metasound
 			// Reset flag to track render cost.
 			bDidWaveRenderThisBlock = false;
 
+#if !UE_BUILD_SHIPPING
+			if (bIsPlaying && WaveAsset->IsSoundWaveValid() && WaveAsset->GetSoundWaveProxy()->TryGetLatest())
+			{
+				StartPlaying();
+			}
+#endif // !UE_BUILD_SHIPPING
+
 			// Update wave proxy reader with any new looping bounds. 
 			if (WaveProxyReader.IsValid())
 			{
diff --git a/Engine/Source/Runtime/AudioExtensions/Public/IAudioProxyInitializer.h b/Engine/Source/Runtime/AudioExtensions/Public/IAudioProxyInitializer.h
--- a/Engine/Source/Runtime/AudioExtensions/Public/IAudioProxyInitializer.h
+++ b/Engine/Source/Runtime/AudioExtensions/Public/IAudioProxyInitializer.h
@@ -4,6 +4,7 @@
 
 #include "Misc/AssertionMacros.h"
 #include "Misc/CoreMiscDefines.h"
+#include "Misc/ScopeTryLock.h"
 #include "Templates/Casts.h"
 #include "Templates/SharedPointer.h"
 #include "Templates/UniquePtr.h"
@@ -11,6 +12,149 @@
 
 #include <type_traits>
 
+
+template <typename T>
+class TDataSnapshotProducer;
+
+template <typename T>
+class TDataSnapshot
+{
+public:
+	TDataSnapshot() = delete;
+	
+	TDataSnapshot(TWeakPtr<const TDataSnapshotProducer<T>> InWeakProducer, TSharedPtr<T> InPublishedData)
+	: WeakProducer(InWeakProducer)
+	, DataSnapshot(InPublishedData)
+	{}
+
+	operator TSharedPtr<T>() const
+	{
+		return DataSnapshot;
+	}
+
+	bool IsValid() const
+	{
+		return DataSnapshot.IsValid();
+	}
+
+	explicit operator bool() const
+	{
+		return DataSnapshot.IsValid();
+	}
+	
+	T* operator->() const
+	{
+		return DataSnapshot.Get();
+	}
+
+	bool IsUpdateAvailable() const
+	{
+		if (TSharedPtr<const TDataSnapshotProducer<T>> PinnedProducer = WeakProducer.Pin())
+		{
+			return PinnedProducer->IsSnapshotUpdateAvailable(DataSnapshot);
+		}
+		
+		return false;
+	}
+
+	bool GetLatest()
+	{
+		TSharedPtr<const TDataSnapshotProducer<T>> PinnedProducer = WeakProducer.Pin();
+
+		if (!PinnedProducer)
+		{
+			// producer is out of scope, no more updates
+			return false;
+		}
+
+		if (PinnedProducer->IsSnapshotUpdateAvailable(DataSnapshot))
+		{
+			// grab new snapshot
+			DataSnapshot = PinnedProducer->GetData();
+		
+			return true;
+		}
+		
+		return false;
+	}
+	
+private:
+	TWeakPtr<const TDataSnapshotProducer<T>> WeakProducer;
+	TSharedPtr<T> DataSnapshot;
+};
+
+template <typename T>
+class TDataSnapshotProducer : public TSharedFromThis<TDataSnapshotProducer<T>>
+{
+public:
+	operator TSharedPtr<T>() const
+	{
+		return WorkingData;
+	}
+
+	bool IsValid() const
+	{
+		return WorkingData.IsValid();
+	}
+	
+	explicit operator bool() const
+	{
+		return WorkingData.IsValid();
+	}
+	
+	T* operator->() const
+	{
+		return WorkingData.Get();
+	}
+
+	void Fork()
+	{
+		FWriteScopeLock Lock(DataLock);
+		WorkingData = MakeShared<T>();
+	}
+
+	// note: there are still some things the USoundWave will potentially update
+	// on existing proxies (i.e. certain flags). Letting these pointers alias 
+	// after "Publish()" is intentional for now.
+	void Publish()
+	{
+		FWriteScopeLock Lock(DataLock);
+		PublishedData = WorkingData;
+	}
+	
+	TDataSnapshot<T> GetSnapshot() const
+	{
+		return TDataSnapshot<T>(this->AsWeak(), GetData());
+	}
+
+	bool IsSnapshotUpdateAvailable(const TSharedPtr<T>& InSnapshot) const
+	{
+		if (DataLock.TryReadLock())
+		{
+			const bool bIsUpdateAvailable = PublishedData.IsValid() && (PublishedData != InSnapshot);
+
+			DataLock.ReadUnlock();
+			return bIsUpdateAvailable;
+		}
+
+		return false;
+	}
+	
+private:
+	TSharedPtr<T> GetData() const
+	{
+		FReadScopeLock Lock(DataLock);
+		return PublishedData;
+	}
+	
+	TSharedPtr<T> PublishedData;
+	TSharedPtr<T> WorkingData = MakeShared<T>();
+	mutable FRWLock DataLock;
+	
+	friend class TDataSnapshot<T>;
+}; // TDataSnapshotProducer
+
+
 /**
  * Interfaces for Audio Proxy Objects 
  * These are used to spawn thread safe instances of UObjects that may be garbage collected on the game thread.
diff --git a/Engine/Source/Runtime/Engine/Classes/Sound/SoundWave.h b/Engine/Source/Runtime/Engine/Classes/Sound/SoundWave.h
--- a/Engine/Source/Runtime/Engine/Classes/Sound/SoundWave.h
+++ b/Engine/Source/Runtime/Engine/Classes/Sound/SoundWave.h
@@ -1560,7 +1560,8 @@ public:
 
 
 private:
-	TSharedPtr<FSoundWaveData, ESPMode::ThreadSafe> SoundWaveDataPtr{ MakeShared<FSoundWaveData>() };
+	TSharedRef<TDataSnapshotProducer<FSoundWaveData>> SoundWaveData = MakeShared<TDataSnapshotProducer<FSoundWaveData>>();
+	TDataSnapshotProducer<FSoundWaveData>& SoundWaveDataPtr = *SoundWaveData;
 
 	friend class FSoundWaveProxy;
 	friend class USoundFactory;
@@ -1815,8 +1816,14 @@ public:
 		return InProxy.TypeHash;
 	}
 	
+	bool TryGetLatest()
+	{
+		return SoundWaveDataPtr.GetLatest();
+	}
+	
 private:
-	TSharedPtr<FSoundWaveData, ESPMode::ThreadSafe> SoundWaveDataPtr;
+	TDataSnapshot<FSoundWaveData> SoundWaveDataPtr;
+	
 	uint32 TypeHash = uint32(INDEX_NONE);
 };
 
diff --git a/Engine/Source/Runtime/Engine/Private/SoundWave.cpp b/Engine/Source/Runtime/Engine/Private/SoundWave.cpp
--- a/Engine/Source/Runtime/Engine/Private/SoundWave.cpp
+++ b/Engine/Source/Runtime/Engine/Private/SoundWave.cpp
@@ -737,14 +737,14 @@ FString FSoundWaveData::GetDerivedDataKey() const
 
 bool USoundWave::IsCompiling() const
 {
-	return SoundWaveDataPtr != nullptr && SoundWaveDataPtr->RunningPlatformData.IsCompiling();
+	return SoundWaveDataPtr && SoundWaveDataPtr->RunningPlatformData.IsCompiling();
 }
 
 /** IInterface_AsyncCompilation implementation end*/
 
 bool USoundWave::IsAsyncWorkComplete() const
 {
-	return SoundWaveDataPtr == nullptr || SoundWaveDataPtr->RunningPlatformData.IsAsyncWorkComplete();
+	return !SoundWaveDataPtr || SoundWaveDataPtr->RunningPlatformData.IsAsyncWorkComplete();
 }
 
 void USoundWave::PostImport()
@@ -3362,7 +3362,6 @@ void USoundWave::PostEditChangeProperty(FPropertyChangedEvent& PropertyChangedEv
 			}
 		}
 	}
-
 }
 
 void USoundWave::CreateNewSoundWaveData()
@@ -3383,7 +3382,7 @@ void USoundWave::CreateNewSoundWaveData()
 	// be extended to support notifying clients that it has been edited and replaced 
 	// an existing proxy with a new proxy. Clients can then choose whether to
 	// continue playing their original proxy, or retrieve the updated proxy. 
-	SoundWaveDataPtr = MakeShared<FSoundWaveData, ESPMode::ThreadSafe>();
+	SoundWaveDataPtr.Fork();
 	Proxy.Reset();
 
 }
@@ -4256,7 +4255,7 @@ FSoundWaveProxyPtr USoundWave::CreateSoundWaveProxy()
 #endif // #if WITH_EDITORONLY_DATA
 
 	check(SoundWaveDataPtr);
-
+	SoundWaveDataPtr.Publish();
 	return MakeShared<FSoundWaveProxy, ESPMode::ThreadSafe>(this);
 }
 
@@ -5224,7 +5223,7 @@ ESoundWaveLoadingBehavior USoundWave::GetLoadingBehavior(bool bCheckSoundClasses
 }
 
 FSoundWaveProxy::FSoundWaveProxy(USoundWave* InWave)
-	: SoundWaveDataPtr(InWave->SoundWaveDataPtr)
+	: SoundWaveDataPtr(InWave->SoundWaveDataPtr.GetSnapshot())
 {
 	LLM_SCOPE(ELLMTag::AudioSoundWaveProxies);
 
