[AudioExtensions]
* Fixed header units compile warning

diff --git a/Engine/Source/Runtime/AudioExtensions/Public/IAudioModulation.h b/Engine/Source/Runtime/AudioExtensions/Public/IAudioModulation.h
--- a/Engine/Source/Runtime/AudioExtensions/Public/IAudioModulation.h
+++ b/Engine/Source/Runtime/AudioExtensions/Public/IAudioModulation.h
@@ -169,9 +169,9 @@ namespace Audio
 
 	private:
 		FModulationParameter Parameter;
-		FModulatorHandleId HandleId = INDEX_NONE;
-		FModulatorTypeId ModulatorTypeId = INDEX_NONE;
-		FModulatorId ModulatorId = INDEX_NONE;
+		FModulatorHandleId HandleId = (FModulatorHandleId)INDEX_NONE;
+		FModulatorTypeId ModulatorTypeId = (FModulatorTypeId)INDEX_NONE;
+		FModulatorId ModulatorId = (FModulatorId)INDEX_NONE;
 		TWeakPtr<IAudioModulationManager> Modulation;
 	};
 } // namespace Audio
