Modulation: Per-Parameter Unclamp CVar
#Synced-CL 48954519

diff --git a/Engine/Plugins/Runtime/AudioModulation/Source/AudioModulation/Private/AudioModulationSubsystem.cpp b/Engine/Plugins/Runtime/AudioModulation/Source/AudioModulation/Private/AudioModulationSubsystem.cpp
new file mode 100644
--- /dev/null
+++ b/Engine/Plugins/Runtime/AudioModulation/Source/AudioModulation/Private/AudioModulationSubsystem.cpp
@@ -0,0 +1,126 @@
+﻿// Copyright Epic Games, Inc. All Rights Reserved.
+#include "AudioModulationSubsystem.h"
+
+#include "AudioDevice.h"
+#include "AudioModulationSystem.h"
+#include "Misc/App.h"
+
+static FAutoConsoleCommand CVarAudioModulationSetParameterClampOverride(
+	TEXT("au.Modulation.SetParameterClampOverride"),
+	TEXT("Sets whether the given parameter should be clamped to the normalized [0,1] range. Overrides the UObject setting, but not the ForceClampAllModulationValues CVar. Arguments:\n"
+		"Parameter - The parameter to set the clamp status of.\n"
+		"Clamp (Optional, Default: 1) - Whether the parameter should be clamped. 1 = clamped, 0 = unclamped.\n"),
+		FConsoleCommandWithArgsDelegate::CreateLambda([](const TArray<FString>& Args)
+	{
+		if (Args.Num() < 1)
+		{
+			return;
+		}
+
+		bool bShouldClamp = true;
+		if (Args.Num() > 1)
+		{
+			if (FCString::IsNumeric(*Args[1]))
+			{
+				bShouldClamp = static_cast<bool>(FCString::Atoi(*Args[1]));
+			}
+		}
+
+		if (FAudioDeviceManager* AudioDeviceManager = FAudioDeviceManager::Get())
+		{
+			AudioDeviceManager->IterateOverAllDevices([ParamName = Args[0], bShouldClamp](Audio::FDeviceId Id, FAudioDevice* AudioDevice )
+			{
+				if (!AudioDevice)
+				{
+					return;
+				}
+				
+				if (UAudioModulationSubsystem* Subsystem = AudioDevice->GetSubsystem<UAudioModulationSubsystem>())
+				{
+					Subsystem->SetModulationParameterClampOverride(ParamName, bShouldClamp);
+				}
+			});
+		}
+	}));
+
+static FAutoConsoleCommand CVarAudioModulationClearParameterClampOverride(
+	TEXT("au.Modulation.ClearParameterClampOverride"),
+	TEXT("Removes any override for the given parameter set by SetParameterClampOverride.\n"
+		"Parameter - The parameter to clear the override of.\n"),
+		FConsoleCommandWithArgsDelegate::CreateLambda([](const TArray<FString>& Args)
+	{
+		if (Args.Num() < 1)
+		{
+			return;
+		}
+
+			
+		if (FAudioDeviceManager* AudioDeviceManager = FAudioDeviceManager::Get())
+		{
+			AudioDeviceManager->IterateOverAllDevices([ParamName = Args[0]](Audio::FDeviceId Id, FAudioDevice* AudioDevice )
+			{
+				if (!AudioDevice)
+				{
+					return;
+				}
+				
+				if (UAudioModulationSubsystem* Subsystem = AudioDevice->GetSubsystem<UAudioModulationSubsystem>())
+				{
+					Subsystem->ClearModulationParameterClampOverride(ParamName);
+				}
+			});
+		}
+	}));
+
+static FAutoConsoleCommand CVarAudioModulationClearAllParameterClampOverrides(
+	TEXT("au.Modulation.ClearAllParameterClampOverrides"),
+	TEXT("Removes all modulation parameter clamp overrides. This does not remove the global clamp set by 'au.Modulation.ForceClampAllModulationValues'. \n"),
+		FConsoleCommandDelegate::CreateLambda([]()
+	{
+		if (FAudioDeviceManager* AudioDeviceManager = FAudioDeviceManager::Get())
+		{
+			AudioDeviceManager->IterateOverAllDevices([](Audio::FDeviceId Id, FAudioDevice* AudioDevice )
+			{
+				if (!AudioDevice)
+				{
+					return;
+				}
+				
+				if (UAudioModulationSubsystem* Subsystem = AudioDevice->GetSubsystem<UAudioModulationSubsystem>())
+				{
+					Subsystem->ClearAllModulationParameterClampOverrides();
+				}
+			});
+		}
+	}));
+
+bool UAudioModulationSubsystem::ShouldCreateSubsystem(UObject* Outer) const
+{
+	return Super::ShouldCreateSubsystem(Outer) && FApp::CanEverRenderAudio();
+}
+
+void UAudioModulationSubsystem::Deinitialize()
+{
+	ModulationParameterClampOverrides.Empty();
+}
+
+void UAudioModulationSubsystem::SetModulationParameterClampOverride(const FString& ParameterName, bool bClamp) const
+{
+	PRAGMA_DISABLE_INTERNAL_WARNINGS
+	Audio::SetModulationParameterClampOverride(FName(*ParameterName), bClamp);
+	PRAGMA_ENABLE_INTERNAL_WARNINGS
+}
+
+void UAudioModulationSubsystem::ClearModulationParameterClampOverride(const FString& ParameterName) const
+{
+	PRAGMA_DISABLE_INTERNAL_WARNINGS
+	Audio::ClearModulationParameterClampOverride(FName(*ParameterName));
+	PRAGMA_ENABLE_INTERNAL_WARNINGS
+}
+
+void UAudioModulationSubsystem::ClearAllModulationParameterClampOverrides() const
+{
+	PRAGMA_DISABLE_INTERNAL_WARNINGS
+	Audio::ClearAllModulationParameterClampOverrides();
+	PRAGMA_ENABLE_INTERNAL_WARNINGS
+}
diff --git a/Engine/Plugins/Runtime/AudioModulation/Source/AudioModulation/Private/AudioModulationSystem.h b/Engine/Plugins/Runtime/AudioModulation/Source/AudioModulation/Private/AudioModulationSystem.h
--- a/Engine/Plugins/Runtime/AudioModulation/Source/AudioModulation/Private/AudioModulationSystem.h
+++ b/Engine/Plugins/Runtime/AudioModulation/Source/AudioModulation/Private/AudioModulationSystem.h
@@ -272,10 +272,6 @@ namespace AudioModulation
 
 		Audio::FDeviceId AudioDeviceId = INDEX_NONE;
 
-		// Registered Parameters
-		using FParameterRegistry = TMap<FName, Audio::FModulationParameter>;
-		FParameterRegistry ParameterRegistry;
-
 #if UE_AUDIO_PROFILERTRACE_ENABLED
 		float ControlBusActivePingTimeSec = 0.0f;
 		static constexpr float ControlBusActivePingIntervalTimeSec = 1.0f;
diff --git a/Engine/Plugins/Runtime/AudioModulation/Source/AudioModulation/Private/SoundControlBusProxy.cpp b/Engine/Plugins/Runtime/AudioModulation/Source/AudioModulation/Private/SoundControlBusProxy.cpp
--- a/Engine/Plugins/Runtime/AudioModulation/Source/AudioModulation/Private/SoundControlBusProxy.cpp
+++ b/Engine/Plugins/Runtime/AudioModulation/Source/AudioModulation/Private/SoundControlBusProxy.cpp
@@ -6,13 +6,14 @@
 #include "AudioMixerTrace.h"
 #include "AudioModulation.h"
 #include "AudioModulationLogging.h"
+#include "AudioModulationSubsystem.h"
 #include "AudioModulationSystem.h"
 #include "Engine/World.h"
 #include "SoundModulationGenerator.h"
 
 static int32 ForceClampNormalizedValuesCVar = 0;
 FAutoConsoleVariableRef CVarForceClampNormalizedValues(
-	TEXT("au.Modulation.ForceClampModulationValues"),
+	TEXT("au.Modulation.ForceClampAllModulationValues"),
 	ForceClampNormalizedValuesCVar,
 	TEXT("Forces all Modulation values to stay in the normalized [0,1] range.\n")
 	TEXT("0: Normalized values can go outside [0,1] if the Modulation Parameter allows it. 1: Normalized values are always in [0,1] range."),
@@ -128,7 +129,7 @@ namespace AudioModulation
 	{
 		float DefaultMixed = Mix(DefaultValue) * GeneratorValue;
 
-		if (ForceClampNormalizedValuesCVar || bClampToNormalizedRange)
+		if (IsParameterClamped())
 		{
 			return FMath::Clamp(DefaultMixed, 0.0f, 1.0f);
 		}
@@ -201,6 +202,18 @@ namespace AudioModulation
 		return OutValue;
 	}
 
+	bool FControlBusProxy::IsParameterClamped() const
+	{
+		// This cvar forces all values to be clamped, highest priority
+		if (ForceClampNormalizedValuesCVar)
+		{
+			return true;
+		}
+
+		// The value set on the original Parameter UObject
+		return bClampToNormalizedRange;		
+	}
+
 #if UE_AUDIO_PROFILERTRACE_ENABLED
 	void AudioModulation::FControlBusProxy::OnTraceStarted(FAudioModulationSystem& InModSystem) const
 	{
diff --git a/Engine/Plugins/Runtime/AudioModulation/Source/AudioModulation/Private/SoundControlBusProxy.h b/Engine/Plugins/Runtime/AudioModulation/Source/AudioModulation/Private/SoundControlBusProxy.h
--- a/Engine/Plugins/Runtime/AudioModulation/Source/AudioModulation/Private/SoundControlBusProxy.h
+++ b/Engine/Plugins/Runtime/AudioModulation/Source/AudioModulation/Private/SoundControlBusProxy.h
@@ -90,7 +90,8 @@ namespace AudioModulation
 	private:
 		void Init(FControlBusSettings&& InSettings);
 		float Mix(float ValueA) const;
-
+		bool IsParameterClamped() const;
+		
 		float DefaultValue = 0.0f;
 		
 		// Cached values
diff --git a/Engine/Plugins/Runtime/AudioModulation/Source/AudioModulation/Public/AudioModulationSubsystem.h b/Engine/Plugins/Runtime/AudioModulation/Source/AudioModulation/Public/AudioModulationSubsystem.h
new file mode 100644
--- /dev/null
+++ b/Engine/Plugins/Runtime/AudioModulation/Source/AudioModulation/Public/AudioModulationSubsystem.h
@@ -0,0 +1,33 @@
+// Copyright Epic Games, Inc. All Rights Reserved.
+
+#pragma once
+
+#include "Subsystems/AudioEngineSubsystem.h"
+
+#include "AudioModulationSubsystem.generated.h"
+
+/**
+*  UAudioModulationSubsystem
+*/
+UCLASS(MinimalAPI)
+class UAudioModulationSubsystem : public UAudioEngineSubsystem
+{
+	GENERATED_BODY()
+
+public:
+	UAudioModulationSubsystem() = default;
+	virtual ~UAudioModulationSubsystem() override = default;
+
+	//~ Begin USubsystem interface
+	AUDIOMODULATION_API virtual bool ShouldCreateSubsystem(UObject* Outer) const override;
+	AUDIOMODULATION_API virtual void Deinitialize() override;
+	//~ End USubsystem interface
+
+	// Modulation Parameter Overrides
+	void SetModulationParameterClampOverride(const FString& ParameterName, bool bClamp) const;
+	void ClearModulationParameterClampOverride(const FString& ParameterName) const;
+	void ClearAllModulationParameterClampOverrides() const;
+
+private:
+	TMap<FString, bool> ModulationParameterClampOverrides;
+};
diff --git a/Engine/Source/Runtime/AudioExtensions/Private/IAudioModulation.cpp b/Engine/Source/Runtime/AudioExtensions/Private/IAudioModulation.cpp
--- a/Engine/Source/Runtime/AudioExtensions/Private/IAudioModulation.cpp
+++ b/Engine/Source/Runtime/AudioExtensions/Private/IAudioModulation.cpp
@@ -20,7 +20,16 @@ namespace Audio
 	{
 		class FModulationParameterRegistry
 		{
-			TMap<FName, TUniquePtr<FModulationParameter>> Values;
+		public:	
+			struct FParameterRegisterEntry
+			{
+				FModulationParameter Parameter;
+				// Clamp Setting from when the Parameter was first added to the Registry
+				bool bRegisteredClampSetting;
+			};
+
+		private:
+			TMap<FName, TUniquePtr<FParameterRegisterEntry>> Values;
 
 			mutable FCriticalSection ThreadSafeValueAccessor;
 
@@ -34,16 +43,18 @@ namespace Audio
 			void Register(FName InName, FModulationParameter&& InParameter)
 			{
 				FScopeLock Lock(&ThreadSafeValueAccessor);
-				if (const TUniquePtr<FModulationParameter>* ParamPtr = Values.Find(InName))
+				if (const TUniquePtr<FParameterRegisterEntry>* ParamEntryPtr = Values.Find(InName))
 				{
-					if (FModulationParameter* Param = ParamPtr->Get())
-					{
-						*Param = MoveTemp(InParameter);
-						return;
-					}
+					const TUniquePtr<FParameterRegisterEntry>& ParamEntry = *ParamEntryPtr;
+					FModulationParameter& Param = ParamEntry->Parameter;
+					Param = MoveTemp(InParameter);
+					return;
 				}
-				
-				Values.Add(InName, MakeUnique<FModulationParameter>(FModulationParameter(InParameter)));
+
+				TUniquePtr<FParameterRegisterEntry> NewEntry = MakeUnique<FParameterRegisterEntry>();
+				NewEntry->Parameter = InParameter;
+				NewEntry->bRegisteredClampSetting = InParameter.bClampToNormalizedRange;
+				Values.Add(InName, MoveTemp(NewEntry));
 			}
 
 			bool Unregister(FName InName)
@@ -58,15 +69,52 @@ namespace Audio
 				Values.Reset();
 			}
 
-			const FModulationParameter* Get(FName InName) const
+			void SetClampOverride(FName InName, bool bClampOverride)
+			{
+				if (TUniquePtr<FParameterRegisterEntry>* ParamEntry = Values.Find(InName))
+				{
+					if (FParameterRegisterEntry* ParamEntryPtr = ParamEntry->Get())
+					{
+						FModulationParameter& Param = ParamEntryPtr->Parameter;
+						Param.bClampToNormalizedRange = bClampOverride;
+					}
+				}
+			}
+
+			void ClearClampOverride(FName InName)
+			{
+				if (TUniquePtr<FParameterRegisterEntry>* ParamEntry = Values.Find(InName))
+				{
+					if (FParameterRegisterEntry* ParamEntryPtr = ParamEntry->Get())
+					{
+						FModulationParameter& Param = ParamEntryPtr->Parameter;
+						Param.bClampToNormalizedRange = ParamEntryPtr->bRegisteredClampSetting;
+					}
+				}
+			}
+
+			void ClearAllClampOverrides()
+			{
+				for (TPair<FName, TUniquePtr<FParameterRegisterEntry>>& Pair : Values)
+				{
+					TUniquePtr<FParameterRegisterEntry>& ParamEntry = Pair.Value;
+					if (FParameterRegisterEntry* ParamEntryPtr = ParamEntry.Get())
+					{
+						FModulationParameter& Param = ParamEntryPtr->Parameter;
+						Param.bClampToNormalizedRange = ParamEntryPtr->bRegisteredClampSetting;
+					}
+				}
+			}
+			
+			const FModulationParameter* GetParameter(FName InName) const
 			{
 				FScopeLock Lock(&ThreadSafeValueAccessor);
 
-				if (const TUniquePtr<FModulationParameter>* ParamPtr = Values.Find(InName))
+				if (const TUniquePtr<FParameterRegisterEntry>* ParamEntry = Values.Find(InName))
 				{
-					if (const FModulationParameter* Param = ParamPtr->Get())
+					if (FParameterRegisterEntry* ParamEntryPtr = ParamEntry->Get())
 					{
-						return Param;
+						return &ParamEntryPtr->Parameter;
 					}
 				}
 
@@ -142,7 +190,7 @@ namespace Audio
 	const FModulationParameter* GetModulationParameterPtr(FName InName)
 	{
 		using namespace ModulationInterfacePrivate;
-		return ParameterRegistry.Get(InName);
+		return ParameterRegistry.GetParameter(InName);
 	}
 
 	const FModulationParameter& GetDefaultModulationParameter()
@@ -150,11 +198,26 @@ namespace Audio
 		return ModulationPrivate::DefaultParameter;
 	}
 
+	void SetModulationParameterClampOverride(FName InName, bool bClamp)
+	{
+		ModulationInterfacePrivate::ParameterRegistry.SetClampOverride(InName, bClamp);
+	}
+
+	void ClearModulationParameterClampOverride(FName InName)
+	{
+		ModulationInterfacePrivate::ParameterRegistry.ClearClampOverride(InName);
+	}
+
+	void ClearAllModulationParameterClampOverrides()
+	{
+		ModulationInterfacePrivate::ParameterRegistry.ClearAllClampOverrides();
+	}
+
 	const FModulationParameter& GetModulationParameter(FName InName)
 	{
 		using namespace ModulationInterfacePrivate;
 
-		if (const FModulationParameter* Param = ParameterRegistry.Get(InName))
+		if (const FModulationParameter* Param = ParameterRegistry.GetParameter(InName))
 		{
 			return *Param;
 		}
diff --git a/Engine/Source/Runtime/AudioExtensions/Public/IAudioModulation.h b/Engine/Source/Runtime/AudioExtensions/Public/IAudioModulation.h
--- a/Engine/Source/Runtime/AudioExtensions/Public/IAudioModulation.h
+++ b/Engine/Source/Runtime/AudioExtensions/Public/IAudioModulation.h
@@ -108,6 +108,15 @@ namespace Audio
 	AUDIOEXTENSIONS_API const FModulationParameter* GetModulationParameterPtr(FName InName);
 	AUDIOEXTENSIONS_API const FModulationParameter& GetDefaultModulationParameter();
 	
+	UE_INTERNAL
+	AUDIOEXTENSIONS_API void SetModulationParameterClampOverride(FName InName, bool bClamp);
+	
+	UE_INTERNAL
+	AUDIOEXTENSIONS_API void ClearModulationParameterClampOverride(FName InName);
+	
+	UE_INTERNAL
+	AUDIOEXTENSIONS_API void ClearAllModulationParameterClampOverrides();
+	
 	UE_DEPRECATED(5.6, "Please use the version of GetModulationParameter which returns a pointer instead of a reference, and check that it is valid in your code.")
 	AUDIOEXTENSIONS_API const FModulationParameter& GetModulationParameter(FName InName);
 
