Fix issues with modulation parameter clamp cvar behavior

#Synced-CL 50781067

diff --git a/Engine/Plugins/Runtime/AudioModulation/Source/AudioModulation/Private/SoundModulationParameter.cpp b/Engine/Plugins/Runtime/AudioModulation/Source/AudioModulation/Private/SoundModulationParameter.cpp
--- a/Engine/Plugins/Runtime/AudioModulation/Source/AudioModulation/Private/SoundModulationParameter.cpp
+++ b/Engine/Plugins/Runtime/AudioModulation/Source/AudioModulation/Private/SoundModulationParameter.cpp
@@ -69,7 +69,7 @@ Audio::FModulationParameter USoundModulationParameter::CreateParameter() const
 	Parameter.DefaultValue = GetUnitDefault();
 	Parameter.MinValue = GetUnitMin();
 	Parameter.MaxValue = GetUnitMax();
-	Parameter.bClampToNormalizedRange = IsModulationParameterClamped();
+	Parameter.bClampToNormalizedRange = Settings.bClampToNormalizedRange;
 	Parameter.TypeHash = GetTypeHash(this);
 
 #if WITH_EDITORONLY_DATA
diff --git a/Engine/Source/Runtime/AudioExtensions/Private/IAudioModulation.cpp b/Engine/Source/Runtime/AudioExtensions/Private/IAudioModulation.cpp
--- a/Engine/Source/Runtime/AudioExtensions/Private/IAudioModulation.cpp
+++ b/Engine/Source/Runtime/AudioExtensions/Private/IAudioModulation.cpp
@@ -20,18 +20,12 @@ namespace Audio
 	{
 		class FModulationParameterRegistry
 		{
-		public:	
-			struct FParameterRegisterEntry
-			{
-				FModulationParameter Parameter;
-				// Clamp Setting from when the Parameter was first added to the Registry
-				bool bRegisteredClampSetting;
-			};
-
 		private:
-			TMap<FName, TUniquePtr<FParameterRegisterEntry>> Values;
+			TMap<FName, TUniquePtr<FModulationParameter>> Values;
+			TMap<FName, bool> ClampOverrides;
 
 			mutable FCriticalSection ThreadSafeValueAccessor;
+			mutable FCriticalSection ThreadSafeClampOverrideAccessor;
 
 		public:
 			bool IsRegistered(FName InName) const
@@ -43,18 +37,16 @@ namespace Audio
 			void Register(FName InName, FModulationParameter&& InParameter)
 			{
 				FScopeLock Lock(&ThreadSafeValueAccessor);
-				if (const TUniquePtr<FParameterRegisterEntry>* ParamEntryPtr = Values.Find(InName))
+				if (const TUniquePtr<FModulationParameter>* ParamPtr = Values.Find(InName))
 				{
-					const TUniquePtr<FParameterRegisterEntry>& ParamEntry = *ParamEntryPtr;
-					FModulationParameter& Param = ParamEntry->Parameter;
-					Param = MoveTemp(InParameter);
-					return;
+					if (FModulationParameter* Param = ParamPtr->Get())
+					{
+						*Param = MoveTemp(InParameter);
+						return;
+					}
 				}
 
-				TUniquePtr<FParameterRegisterEntry> NewEntry = MakeUnique<FParameterRegisterEntry>();
-				NewEntry->Parameter = InParameter;
-				NewEntry->bRegisteredClampSetting = InParameter.bClampToNormalizedRange;
-				Values.Add(InName, MoveTemp(NewEntry));
+				Values.Add(InName, MakeUnique<FModulationParameter>(FModulationParameter(InParameter)));
 			}
 
 			bool Unregister(FName InName)
@@ -71,50 +63,40 @@ namespace Audio
 
 			void SetClampOverride(FName InName, bool bClampOverride)
 			{
-				if (TUniquePtr<FParameterRegisterEntry>* ParamEntry = Values.Find(InName))
+				FScopeLock Lock(&ThreadSafeClampOverrideAccessor);
+				if (ClampOverrides.Contains(InName))
 				{
-					if (FParameterRegisterEntry* ParamEntryPtr = ParamEntry->Get())
-					{
-						FModulationParameter& Param = ParamEntryPtr->Parameter;
-						Param.bClampToNormalizedRange = bClampOverride;
-					}
+					ClampOverrides[InName] = bClampOverride;
+				}
+				else
+				{
+					ClampOverrides.Add(InName, bClampOverride);
 				}
 			}
 
 			void ClearClampOverride(FName InName)
 			{
-				if (TUniquePtr<FParameterRegisterEntry>* ParamEntry = Values.Find(InName))
+				FScopeLock Lock(&ThreadSafeClampOverrideAccessor);
+				if (ClampOverrides.Contains(InName))
 				{
-					if (FParameterRegisterEntry* ParamEntryPtr = ParamEntry->Get())
-					{
-						FModulationParameter& Param = ParamEntryPtr->Parameter;
-						Param.bClampToNormalizedRange = ParamEntryPtr->bRegisteredClampSetting;
-					}
+					ClampOverrides.Remove(InName);
 				}
 			}
 
 			void ClearAllClampOverrides()
 			{
-				for (TPair<FName, TUniquePtr<FParameterRegisterEntry>>& Pair : Values)
-				{
-					TUniquePtr<FParameterRegisterEntry>& ParamEntry = Pair.Value;
-					if (FParameterRegisterEntry* ParamEntryPtr = ParamEntry.Get())
-					{
-						FModulationParameter& Param = ParamEntryPtr->Parameter;
-						Param.bClampToNormalizedRange = ParamEntryPtr->bRegisteredClampSetting;
-					}
-				}
+				FScopeLock Lock(&ThreadSafeClampOverrideAccessor);
+				ClampOverrides.Empty();
 			}
 			
-			const FModulationParameter* GetParameter(FName InName) const
+			const FModulationParameter* Get(FName InName) const
 			{
 				FScopeLock Lock(&ThreadSafeValueAccessor);
-
-				if (const TUniquePtr<FParameterRegisterEntry>* ParamEntry = Values.Find(InName))
+				if (const TUniquePtr<FModulationParameter>* ParamPtr = Values.Find(InName))
 				{
-					if (FParameterRegisterEntry* ParamEntryPtr = ParamEntry->Get())
+					if (const FModulationParameter* Param = ParamPtr->Get())
 					{
-						return &ParamEntryPtr->Parameter;
+						return Param;
 					}
 				}
 
@@ -124,12 +106,19 @@ namespace Audio
 			bool GetParameterClamp(FName InName) const
 			{
 				FScopeLock Lock(&ThreadSafeValueAccessor);
+				FScopeLock ClampLock(&ThreadSafeClampOverrideAccessor);
 
-				if (const TUniquePtr<FParameterRegisterEntry>* ParamEntry = Values.Find(InName))
+				// First check for overrides
+				if (ClampOverrides.Contains(InName))
+				{
+					return ClampOverrides[InName];
+				}
+				
+				if (const TUniquePtr<FModulationParameter>* ParamPtr = Values.Find(InName))
 				{
-					if (FParameterRegisterEntry* ParamEntryPtr = ParamEntry->Get())
+					if (FModulationParameter* Param = ParamPtr->Get())
 					{
-						return ParamEntryPtr->Parameter.bClampToNormalizedRange;
+						return Param->bClampToNormalizedRange;
 					}
 				}
 
@@ -205,7 +194,7 @@ namespace Audio
 	const FModulationParameter* GetModulationParameterPtr(FName InName)
 	{
 		using namespace ModulationInterfacePrivate;
-		return ParameterRegistry.GetParameter(InName);
+		return ParameterRegistry.Get(InName);
 	}
 
 	const FModulationParameter& GetDefaultModulationParameter()
@@ -243,7 +232,7 @@ namespace Audio
 	{
 		using namespace ModulationInterfacePrivate;
 
-		if (const FModulationParameter* Param = ParameterRegistry.GetParameter(InName))
+		if (const FModulationParameter* Param = ParameterRegistry.Get(InName))
 		{
 			return *Param;
 		}
