CL B: IUnknown + IProxyData rework

Add Audio::IUnknown interface (AudioUnknown.h) with virtual QueryInterface
and typed free-function Audio::QueryInterface<T>() templates that avoid
exposing void* to callers.

Rework IProxyData to inherit from IUnknown and unify downcasting:
- IMPL_AUDIOPROXY_CLASS macro now generates QueryInterface override and
GetInterfaceId() alias so proxy types satisfy CHasInterfaceIdFunction
- GetAs<> delegates to QueryInterface internally
- Added using IUnknown::QueryInterface to prevent const overload hiding
- Removed CheckTypeCast (QueryInterface handles type matching)

Add ISoundWaveContainer interface using GetInterfaceId() pattern.

Add unit tests (UnitTestAudioUnknown.cpp): 12 test cases covering IUnknown
QueryInterface, typed Audio::QueryInterface<T>(), IMPL_AUDIOPROXY_CLASS
GetInterfaceId, IProxyData::GetAs<>, and cross-type rejection.

Blast radius: 3 classes use IMPL_AUDIOPROXY_CLASS (SoundWave.h,
IAudioModulation.h, MediaPlayer.h) - they get the new overrides
automatically via macro expansion. GetAs<> has no external callers.

#tests LowLevelTests
#Synced-CL 51594399

diff --git a/Engine/Source/Runtime/AudioExtensions/Public/AudioUnknown.h b/Engine/Source/Runtime/AudioExtensions/Public/AudioUnknown.h
new file mode 100644
--- /dev/null
+++ b/Engine/Source/Runtime/AudioExtensions/Public/AudioUnknown.h
@@ -0,0 +1,48 @@
+// Copyright Epic Games, Inc. All Rights Reserved.
+#pragma once
+
+#include "UObject/NameTypes.h"
+#include "Concepts/ConvertibleTo.h"
+
+namespace Audio
+{
+	/**
+	 * IUnknown like interface for querying unknown data objects of their interfaces.
+	 * https://en.wikipedia.org/wiki/IUnknown
+	 */
+	struct IUnknown
+	{
+		virtual ~IUnknown() = default;
+		virtual const void* QueryInterface(const FName InInterfaceName) const
+		{
+			return const_cast<IUnknown*>(this)->QueryInterface(InInterfaceName);
+		}
+		virtual void* QueryInterface(const FName InInterfaceName) = 0;
+	};
+	
+	template<typename T>
+	concept CHasInterfaceIdFunction = requires 
+	{
+		{ T::GetInterfaceId() } -> UE::CConvertibleTo<FName>;
+	};
+			
+	template <CHasInterfaceIdFunction TInterface>
+	const TInterface* QueryInterface(const IUnknown* InUnknown)
+	{
+		if (InUnknown)
+		{
+			return static_cast<const TInterface*>(InUnknown->QueryInterface(TInterface::GetInterfaceId()));
+		}
+		return nullptr;
+	}
+
+	template <CHasInterfaceIdFunction TInterface>
+	TInterface* QueryInterface(IUnknown* InUnknown)
+	{
+		if (InUnknown)
+		{
+			return static_cast<TInterface*>(InUnknown->QueryInterface(TInterface::GetInterfaceId()));
+		}
+		return nullptr;
+	}
+}
\ No newline at end of file
diff --git a/Engine/Source/Runtime/AudioExtensions/Public/IAudioProxyInitializer.h b/Engine/Source/Runtime/AudioExtensions/Public/IAudioProxyInitializer.h
--- a/Engine/Source/Runtime/AudioExtensions/Public/IAudioProxyInitializer.h
+++ b/Engine/Source/Runtime/AudioExtensions/Public/IAudioProxyInitializer.h
@@ -2,6 +2,7 @@
 
 #pragma once
 
+#include "AudioUnknown.h"
 #include "Misc/AssertionMacros.h"
 #include "Misc/CoreMiscDefines.h"
 #include "Misc/ScopeTryLock.h"
@@ -168,6 +169,15 @@ private:
 		static FName MyClassName = #FClassName; \
 		return MyClassName; \
 	} \
+	static FName GetInterfaceId() { return GetAudioProxyTypeName(); } \
+	virtual void* QueryInterface(const FName InterfaceId) override \
+	{ \
+		if (InterfaceId == GetAudioProxyTypeName()) \
+		{ \
+			return this; \
+		} \
+		return IProxyData::QueryInterface(InterfaceId); \
+	} \
 	static constexpr bool bWasAudioProxyClassImplemented = true; \
 	friend class ::Audio::IProxyData; \
 	friend class ::Audio::TProxyData<FClassName>;
@@ -182,20 +192,20 @@ namespace Audio
 	/*
 	 * Base class that allows us to typecheck proxy data before downcasting it in debug builds.
 	*/
-	class IProxyData
+	class IProxyData : public IUnknown
 	{
 	private:
 		FName ProxyTypeName;
 	public:
 		virtual ~IProxyData() = default;
 
-		template<typename ProxyType>
-		bool CheckTypeCast() const
+		// IUnknown.
+		using IUnknown::QueryInterface;
+		virtual void* QueryInterface(const FName InterfaceId) override 
 		{
-			const FName DestinationTypeName = ProxyType::GetAudioProxyTypeName();
-			return ensureAlwaysMsgf(ProxyTypeName == DestinationTypeName, TEXT("Tried to downcast type %s to %s!"), *ProxyTypeName.ToString(), *DestinationTypeName.ToString());
+			return nullptr;	
 		}
-
+					
 		FName GetProxyTypeName() const
 		{
 			return ProxyTypeName;
@@ -205,32 +215,24 @@ namespace Audio
 		ProxyType& GetAs()
 		{
 			static_assert(std::is_base_of_v<IProxyData, ProxyType>, "Tried to downcast IProxyInitData to an unrelated type!");
-			if (CheckTypeCast<ProxyType>())
-			{
-				return static_cast<ProxyType&>(*this);
-			}
-			else
-			{
-				// This is an illegal cast, and is considered a fatal error.
-				checkNoEntry();
-				return *((ProxyType*)0x1);
-			}
+			ProxyType* Result = static_cast<ProxyType*>(QueryInterface(ProxyType::GetAudioProxyTypeName()));
+			checkf(Result, TEXT("Tried to downcast type %s to %s!"), *GetProxyTypeName().ToString(), *ProxyType::GetAudioProxyTypeName().ToString());
+			return *Result;
 		}
 
 		template<typename ProxyType>
 		const ProxyType& GetAs() const
 		{
 			static_assert(std::is_base_of_v<IProxyData, ProxyType>, "Tried to downcast IProxyInitData to an unrelated type!");
-			if (CheckTypeCast<ProxyType>())
-			{
-				return static_cast<const ProxyType&>(*this);
-			}
-			else
-			{
-				// This is an illegal cast, and is considered a fatal error.
-				checkNoEntry();
-				return *((ProxyType*)0x1);
-			}
+			const ProxyType* Result = static_cast<const ProxyType*>(QueryInterface(ProxyType::GetAudioProxyTypeName()));
+			checkf(Result, TEXT("Tried to downcast type %s to %s!"), *GetProxyTypeName().ToString(), *ProxyType::GetAudioProxyTypeName().ToString());
+			return *Result;
+		}
+
+		template<typename ProxyType>
+		bool CheckTypeCast() const
+		{
+			return QueryInterface(ProxyType::GetAudioProxyTypeName()) != nullptr;
 		}
 
 		IProxyData(FName InProxyTypeName)
diff --git a/Engine/Source/Runtime/Engine/Classes/Sound/ISoundWaveContainer.h b/Engine/Source/Runtime/Engine/Classes/Sound/ISoundWaveContainer.h
new file mode 100644
--- /dev/null
+++ b/Engine/Source/Runtime/Engine/Classes/Sound/ISoundWaveContainer.h
@@ -0,0 +1,25 @@
+// Copyright Epic Games, Inc. All Rights Reserved.
+#pragma once
+
+#include "AudioUnknown.h"
+#include "SoundWave.h"
+#include "UObject/NameTypes.h"
+
+namespace Audio
+{
+	struct UE_EXPERIMENTAL(5.8, "Wave Container Interface is experimental") ISoundWaveContainer
+	{
+		static FName GetInterfaceId()
+		{
+			static const FName Name(TEXT("ISoundWaveContainer"));
+			return Name;
+		}
+
+		virtual ~ISoundWaveContainer() = default;
+
+		virtual TArray<FSoundWaveProxyPtr> GetContainedWaveProxies() const = 0;
+
+		// Validate static API implemented correctly.
+		static_assert(Audio::CHasInterfaceIdFunction<ISoundWaveContainer>);
+	};
+}
